Added the first version of the Celestial Vault DaySequence plugin

#virtualized

diff --git a/Engine/Plugins/Experimental/CelestialVault/Source/CelestialVault/Private/CelestialDataTypes.cpp b/Engine/Plugins/Experimental/CelestialVault/Source/CelestialVault/Private/CelestialDataTypes.cpp
new file mode 100644
--- /dev/null
+++ b/Engine/Plugins/Experimental/CelestialVault/Source/CelestialVault/Private/CelestialDataTypes.cpp
@@ -0,0 +1,77 @@
+﻿// Copyright Epic Games, Inc. All Rights Reserved.
+
+
+#include "CelestialDataTypes.h"
+#include "CelestialMaths.h"
+
+// Earth Preset
+FPlanetaryBodyInputData FPlanetaryBodyInputData::Earth(TEXT("Earth"), EOrbitType::Earth, 6378.0);
+// Moon Preset
+FPlanetaryBodyInputData FPlanetaryBodyInputData::Moon(TEXT("Moon"), EOrbitType::Moon, 1737.4);
+
+
+FString FPlanetaryBodyInfo::ToString() const
+{
+	FString PlanetaryBodyOrbitTypeName = FString("Invalid");
+	if (UEnum* EnumClass = StaticEnum<EOrbitType>())
+	{
+		PlanetaryBodyOrbitTypeName = EnumClass->GetNameStringByValue(static_cast<int64>(OrbitType));
+	}
+
+	return FString::Printf(TEXT(
+			"Planetary Body: %s\n"
+			"Orbit Type: %s\n"
+			"RightAscensionHours: %s\n"
+			"Declination: %s\n"
+			"Distance: %.0f AU\n"
+			"Radius: %.0f km\n"
+			"Apparent Diameter: %s / %.8f°\n"
+			"Apparent Magnitude: %.2f\n"
+			"Phase: %.2f\n"
+			"-------\n"
+			"Scaled Apparent Diameter: %.8f°\n"
+			),
+			*Name,
+			*PlanetaryBodyOrbitTypeName,
+			*UCelestialMaths::Conv_RightAscensionToString(RA),
+			*UCelestialMaths::Conv_DeclinationToString(DEC),
+			DistanceInAU,
+			Radius,
+			*UCelestialMaths::Conv_DeclinationToString(ApparentDiameterDegrees),
+			ApparentDiameterDegrees,
+			ApparentMagnitude,
+			Phase, 
+			ScaledApparentDiameterDegrees
+		); 
+	
+}
+
+FString FSunInfo::ToString() const
+{
+	return FString::Printf(TEXT(
+			"RightAscensionHours: %s\n"
+			"Declination: %s\n"
+			),
+			*UCelestialMaths::Conv_RightAscensionToString(RA),
+			*UCelestialMaths::Conv_DeclinationToString(DEC)
+		);
+}
+
+
+FString FStarInfo::ToString() const
+{
+	return FString::Printf(TEXT(
+			"Name: %s\n"
+			"RightAscensionHours: %s\n"
+			"Declination: %s\n"
+			"Distance: %.0f PC\n"
+			"Magnitude: %.2f\n"
+			),
+			*Name,
+			*UCelestialMaths::Conv_RightAscensionToString(RA),
+			*UCelestialMaths::Conv_DeclinationToString(DEC),
+			DistanceInPC,
+			Magnitude
+		);
+	
+}
diff --git a/Engine/Plugins/Experimental/CelestialVault/Source/CelestialVault/Private/CelestialMaths.cpp b/Engine/Plugins/Experimental/CelestialVault/Source/CelestialVault/Private/CelestialMaths.cpp
new file mode 100644
--- /dev/null
+++ b/Engine/Plugins/Experimental/CelestialVault/Source/CelestialVault/Private/CelestialMaths.cpp
@@ -0,0 +1,1026 @@
+﻿// Copyright Epic Games, Inc. All Rights Reserved.
+
+
+#include "CelestialMaths.h"
+#include "CelestialVault.h"
+#include "MathUtil.h"
+#include "vsop87a_milli.h"
+#include "vsop87a_full_EarthMoon.h"
+
+#pragma region Static Members
+
+// From VSOP87.doc
+// REFERENCE SYSTEM
+// ================
+//
+// The coordinates of the main version VSOP87 and of the versions A, B, and E
+// are given in the inertial frame defined by the dynamical equinox and ecliptic
+// J2000 (JD2451545.0).
+//
+// The rectangular coordinates of VSOP87A and VSOP87E defined in dynamical ecliptic
+// frame J2000 can be connected to the equatorial frame FK5 J2000 with the
+// following rotation :
+//
+//   X        +1.000000000000  +0.000000440360  -0.000000190919   X
+//   Y     =  -0.000000479966  +0.917482137087  -0.397776982902   Y
+//   Z FK5     0.000000000000  +0.397776982902  +0.917482137087   Z VSOP87A
+// 
+FMatrix UCelestialMaths::VSOPToJ2000 = FMatrix( 
+		FVector(+1.000000000000, +0.000000440360, -0.000000190919),	
+		FVector(-0.000000479966, +0.917482137087, -0.397776982902),	
+		FVector(+0.000000000000, +0.397776982902, +0.917482137087),		 
+		FVector(+0.000000000000, 0.000000000000, +0.000000000000)).GetTransposed();	// No Origin offset
+
+double UCelestialMaths::SpeedOfLightMetersPerSeconds = 299792458;
+
+double UCelestialMaths::AstronomicalUnitsMeters = 149597870700;
+
+#pragma endregion
+
+#pragma region Colors
+
+FLinearColor UCelestialMaths::BVtoLinearColor(float BV)
+{
+	// From https://en.wikipedia.org/wiki/Color_index
+
+	// Model valid only between [-0.4, 2.0] 
+	BV = FMath::Clamp(BV, -0.4, 2.0);
+
+	// Compute Effective Temperature
+	const double K1 = 0.92 * BV + 1.7;
+	const double K2 = 0.92 * BV + 0.62;
+	double Temperature = 4600.0 * (1.0 / K1 + 1.0 / K2);
+
+	// Convert to Color
+	return FLinearColor::MakeFromColorTemperature( Temperature);;
+}
+
+#pragma endregion
+
+#pragma region Planetary Bodies
+
+FVector UCelestialMaths::GetBodyLocation_FK5J2000_AU(FPlanetaryBodyInputData PlanetaryBody, double JulianDate)
+{
+	// From VSOP87.doc
+	//   Being given a Julian date JD expressed in dynamical time (TAI+32.184s) and a body (planets, Earth-Moon Barycenter, or Sun) associated to a version of the theory VSOP87 :
+	//     1/ select the file corresponding to the body and the version,
+	//     2/ read sequentially the terms of the series in the records of the file,
+	//     3/ apply for each term the formulae (1) or (2) with T=(JD-2451545)/365250,
+	//     4/ add up the terms so computed for every one coordinate.
+
+	// Convert time, because, the VSOP87 coordinates expect the time in TerrestrialTime
+	double TAI = JulianDateToInternationalAtomicTime(JulianDate); 
+	double TT = InternationalAtomicTimeToTerrestrialTime(TAI); // adds the 32.184s
+	double Time = JulianDateToJulianCenturies(SecondsToDay(TT)) / 10.0; // We divide by 10 because VSOP expects T : time expressed in Thousands of Julian Years (tjy) elapsed from J2000 (JD2451545.0).
+
+	
+	// Get the body coordinates in the inertial frame defined by the dynamical equinox and ecliptic J2000 (JD2451545.0).
+	double BodyXYZ[3] = { 0, 0, 0 }; 
+	switch (PlanetaryBody.OrbitType)
+	{
+	case EOrbitType::Mercury:
+		vsop87a_milli::getMercury(Time, BodyXYZ);
+		break;
+	case EOrbitType::Venus:
+		vsop87a_milli::getVenus(Time, BodyXYZ);
+		break;
+	case EOrbitType::Earth:
+		vsop87a_full::getEarth(Time, BodyXYZ);
+		break;
+	case EOrbitType::Mars:
+		vsop87a_milli::getMars(Time, BodyXYZ);
+		break;
+	case EOrbitType::Jupiter:
+		vsop87a_milli::getJupiter(Time, BodyXYZ);
+		break;
+	case EOrbitType::Saturn:
+		vsop87a_milli::getSaturn(Time, BodyXYZ);
+		break;
+	case EOrbitType::Uranus:
+		vsop87a_milli::getUranus(Time, BodyXYZ);
+		break;
+	case EOrbitType::Neptune:
+		vsop87a_milli::getNeptune(Time, BodyXYZ);
+		break;
+	case EOrbitType::Moon:
+		{
+			// Special case for the Moon. VSOP Works by getting the Earth/Moon barycenter, and combine it with the earth XYZ
+			double EarthXYZ[3] = { 0, 0, 0 }; 
+			vsop87a_full::getEarth(Time, EarthXYZ);
+			double EMBXYZ[3] = { 0, 0, 0 };
+			vsop87a_full::getEmb(Time, EMBXYZ); 
+			vsop87a_full::getMoon(EarthXYZ, EMBXYZ, BodyXYZ);	
+		}
+	case EOrbitType::Elliptic:
+		{
+			// TODO - Use ellipsoid equation to compute locations
+		}
+		break;
+	}
+
+	// Transform these VSOP coordinates into Equatorial Rectangular Coordinates (X, Y, Z)
+	// The FK5 is an equatorial coordinate system (coordinate system linked to the Earth) based on its J2000 position.
+	return VSOPToJ2000.TransformVector(FVector(BodyXYZ[0], BodyXYZ[1], BodyXYZ[2]));
+}
+
+FVector UCelestialMaths::GetBodyLocation_FK5J2000_AU_Relativistic(FVector ObserverBodyFK5J2000LocationAU, FPlanetaryBodyInputData PlanetaryBody, double JulianDate)
+{
+	double JulianDateLightAdjusted = JulianDate;
+	FVector BodyPositionAU = FVector::ZeroVector;
+
+	for (int32 i = 0; i < 3; i++) // 3 iterations are good enough to converge
+	{
+		BodyPositionAU = GetBodyLocation_FK5J2000_AU(PlanetaryBody, JulianDateLightAdjusted);
+		double PlanetaryBodyDistanceAU = FVector::Distance(ObserverBodyFK5J2000LocationAU, BodyPositionAU);
+		double lightPropagationTimeInDays = SecondsToDay( AstronomicalUnitsToMeters(PlanetaryBodyDistanceAU) / SpeedOfLightMetersPerSeconds);
+		JulianDateLightAdjusted = JulianDate - lightPropagationTimeInDays;
+	}
+	return BodyPositionAU;
+}
+
+void UCelestialMaths::GetBodyCelestialCoordinatesAU(double JulianDate, FPlanetaryBodyInputData PlanetaryBody, double ObserverLatitude, double ObserverLongitude, double& RAHours, double& DECDegrees,
+	double& DistanceBodyToEarthAU, double& DistanceBodyToSunAU, double& DistanceEarthToSunAU)
+{
+	TRACE_CPUPROFILER_EVENT_SCOPE(UCelestialMaths::GetBodyCelestialCoordinatesAU);
+
+	
+	RAHours = 0;
+	DECDegrees = 0;
+	DistanceBodyToEarthAU = 0; 
+	
+	if (PlanetaryBody.OrbitType == EOrbitType::Earth)
+	{
+		// It makes no sense to ask for Earth location relative to the Earth, but handle the case anyway...
+		DistanceBodyToSunAU = AstronomicalUnitsMeters; 
+		return;
+	}
+
+	// TODO-Beta
+	// Take into account the Nutation and Precession effect, for better precisions, and also to have the right Heliocentric coordinates, otherwise the Phase Angle computations are wrong.
+	// (Not too important for Moon Phase yet)
+	
+	// Get Earth Location in FK5J2000 Rectangular Coordinates (X, Y, Z)
+	FVector EarthLocationFK5J2000AU = GetBodyLocation_FK5J2000_AU(FPlanetaryBodyInputData::Earth, JulianDate);
+	UE_LOG(LogCelestialVault, Verbose, TEXT("Earth Location : FK5J2000 = %s"), *GetPreciseVectorString(EarthLocationFK5J2000AU));
+
+	// Get Body Location in FK5J2000 Rectangular Coordinates (X, Y, Z)
+	FVector BodyLocationFK5J2000AU = GetBodyLocation_FK5J2000_AU_Relativistic(EarthLocationFK5J2000AU, PlanetaryBody, JulianDate);
+	DistanceBodyToSunAU = BodyLocationFK5J2000AU.Length();
+	UE_LOG(LogCelestialVault, Verbose, TEXT("%s Location : FK5J2000 = %s"), *PlanetaryBody.Name, *GetPreciseVectorString(BodyLocationFK5J2000AU));
+
+	FVector BodyLocationGeoCentricAU = BodyLocationFK5J2000AU - EarthLocationFK5J2000AU;
+	FVector ObserverLocationGeocentricAU = GetObserverGeocentricLocationAU(ObserverLatitude, ObserverLongitude, 0, JulianDate );
+	FVector TopoCentricTargetAU = BodyLocationGeoCentricAU - ObserverLocationGeocentricAU;
+	DistanceEarthToSunAU = (ObserverLocationGeocentricAU-EarthLocationFK5J2000AU).Length();
+	
+	double RADegrees = 0;
+	XYZToRADEC_RH(TopoCentricTargetAU, RADegrees, DECDegrees, DistanceBodyToEarthAU);
+	RAHours = RADegrees / 15.0;
+
+	// Debug Log...
+	UE_LOG(LogCelestialVault, Verbose,
+		TEXT("Planetary Body %s Location : RA = %s, DEC = %s, Radius = %f UA"),
+		*PlanetaryBody.Name,
+		*Conv_RightAscensionToString(RAHours),
+		*Conv_DeclinationToString(DECDegrees),
+		DistanceBodyToEarthAU );
+}
+
+double UCelestialMaths::GetPlanetaryBodyMagnitude(FPlanetaryBodyInputData PlanetaryBody, double DistanceToSunAU, double DistanceToEarthAU, double DistanceEarthToSunAU, double& PhaseAngle)
+{
+	// From Computing Apparent Planetary Magnitudes for The Astronomical Almanac
+	// James L. Hilton US Naval Observatory
+	
+	// V = 5 log10 ( r d ) + V1(0) + C1 α + C2 α2 + ... with
+	//   r = planet’s distance from the Sun
+	//   d = planet’s distance from the earth
+	//	 a = illumination phase angle (in degrees)
+	//   V1(0) sometimes referred to as the planet’s absolute magnitude or geometric magnitude is the magnitude when observed at α = 0
+	//   ΣnCn αn is called the phase function 
+	
+	
+	double DistanceFactor = 5 * FMath::LogX(10, DistanceToEarthAU * DistanceToSunAU);
+	double ApparentMagnitude = 0; 
+	PhaseAngle = FMath::RadiansToDegrees(FMath::Acos( (DistanceToSunAU * DistanceToSunAU + DistanceToEarthAU * DistanceToEarthAU - DistanceEarthToSunAU*DistanceEarthToSunAU) / (2 * DistanceToSunAU * DistanceToEarthAU) ));
+	double PhaseAngle2 = PhaseAngle * PhaseAngle;
+	double PhaseAngle3 = PhaseAngle2 * PhaseAngle;
+	double PhaseAngle4 = PhaseAngle3 * PhaseAngle;
+	double PhaseAngle5 = PhaseAngle4 * PhaseAngle;
+	double PhaseAngle6 = PhaseAngle5 * PhaseAngle;
+	double PhaseFunction = 0; 
+	
+	switch (PlanetaryBody.OrbitType)
+	{
+	case EOrbitType::Elliptic:
+		// Fictional Body. Return dummy magnitude
+		return 0;
+		break;
+	
+	case EOrbitType::Mercury:
+		ApparentMagnitude = -0.613;
+		PhaseFunction = 6.3280E-02 * PhaseAngle - 1.6336E-03 * PhaseAngle2 + 3.3644E-05 * PhaseAngle3 - 3.4265E-07 * PhaseAngle4 + 1.6893E-09 * PhaseAngle5 - 3.0334E-12 *PhaseAngle6;
+		break;
+	case EOrbitType::Venus:
+		if ( PhaseAngle < 163.7)
+		{
+			ApparentMagnitude = -4.384;
+			PhaseFunction = -1.044E-03 * PhaseAngle + 3.687E-04 * PhaseAngle2  - 2.814E-06 * PhaseAngle3 + 8.938E-09 * PhaseAngle4;
+		}
+		else // 163.7 <  α < 179 - let's go to 180...
+		{
+			ApparentMagnitude = 236.05828;
+			PhaseFunction = -2.81914 * PhaseAngle + 8.39034E-03 * PhaseAngle2;
+		}
+		break;
+	case EOrbitType::Earth:
+		ApparentMagnitude = -3.99;
+		PhaseFunction = -1.060E-3 * PhaseAngle + 2.054E-4 * PhaseAngle2;
+		break;
+	case EOrbitType::Mars:
+		if ( PhaseAngle < 50)
+		{
+			ApparentMagnitude = -1.601;
+			PhaseFunction =  0.02267 * PhaseAngle - 0.0001302 * PhaseAngle2;
+		}
+		else 
+		{
+			ApparentMagnitude = -0.367;
+			PhaseFunction = - 0.02573 * PhaseAngle + 0.0003445 * PhaseAngle2;
+		}
+		break;
+	case EOrbitType::Jupiter:
+		if ( PhaseAngle < 12)
+		{
+			ApparentMagnitude =  -9.395;
+			PhaseFunction =  -3.7E-04 * PhaseAngle - 6.16E-04 * PhaseAngle2;
+		}
+		else // 12 < α < 130 - The phase curve of Jupiter as seen from Earth cannot exceed α = 12 so we should be good. Add this one just in case...
+		{
+			ApparentMagnitude = -9.428;
+			PhaseFunction = -2.5 * FMath::LogX(10,
+				1.0
+				-1.507 * (PhaseAngle/180.0)
+				-0.363 * FMath::Pow(PhaseAngle / 180.0, 2)
+				-0.062 * FMath::Pow(PhaseAngle / 180.0, 3)
+				+2.809 * FMath::Pow(PhaseAngle / 180.0, 4)
+				-1.876 * FMath::Pow(PhaseAngle / 180.0, 5)); 
+		}
+		break;
+	case EOrbitType::Saturn:
+		// Keep it simple and ignore the ring effects
+		if ( PhaseAngle < 6)
+		{
+			ApparentMagnitude =  -8.95;
+			PhaseFunction =  - 3.7E-04 * PhaseAngle + 6.16E-04 * PhaseAngle2;
+		}
+		else // 6 < α < 150 -
+		{
+			ApparentMagnitude = - 8.94;
+			PhaseFunction = 2.446E-4 * PhaseAngle + 2.672E-4 * PhaseAngle2 - 1.505E-6 * PhaseAngle3 + 4.767E-9 * PhaseAngle4;
+		}
+		break;
+	case EOrbitType::Uranus:
+		ApparentMagnitude = -7.19;
+		PhaseFunction =  -8.4E-04 * 82;  // PhaseAngle doesn't have any impact, and 82 is the most important planetographic latitude
+		break;
+	case EOrbitType::Neptune:
+		ApparentMagnitude = -7.00;
+		PhaseFunction =  7.944E-3 * PhaseAngle + 9.617E-5 * PhaseAngle2;
+		break;
+	case EOrbitType::Moon:
+		// TODO Alban - A verifier
+		ApparentMagnitude = -12.73;
+		PhaseFunction = 0.026 * PhaseAngle + 4E-9 * PhaseAngle4;
+		break;
+	default: ;
+	}
+
+	return DistanceFactor + ApparentMagnitude + PhaseFunction;
+}
+
+#pragma endregion
+
+#pragma region Time
+
+FDateTime UCelestialMaths::LocalTimeToUTCTime(FDateTime LocalTime, double TimeZoneOffset, bool IsDst)
+{
+	if (IsDst)
+	{
+		TimeZoneOffset += 1.0;
+	}
+
+	return LocalTime - FTimespan(TimeZoneOffset,0,0);
+}
+
+FDateTime UCelestialMaths::UTCTimeToLocalTime(FDateTime UTCTime, double TimeZoneOffset, bool IsDst)
+{
+	if (IsDst)
+	{
+		TimeZoneOffset += 1.0;
+	}
+
+	return UTCTime + FTimespan(TimeZoneOffset,0,0);
+}
+
+double UCelestialMaths::UTCDateTimeToJulianDate(FDateTime UTCDateTime)
+{
+	// From https://www.celestialprogramming.com/julian.html
+
+	// Get Individual values for YMD and HMS
+	int32 Year;
+	int32 Month;
+	int32 Day;
+	UTCDateTime.GetDate(Year, Month, Day);
+	const FTimespan Time = UTCDateTime.GetTimeOfDay();
+	const int32 Hours = Time.GetHours();
+	const int32 Minutes = Time.GetMinutes();
+	const double Seconds = Time.GetTotalSeconds() - Minutes * 60.0 - Hours * 3600.0;
+
+	// Prepare the Input DateTime for JulianDate computations
+	bool bIsGregorian = true;
+	if (Year < 1582 || (Year == 1582 && (Month < 10 || (Month == 10 && Day < 5))))
+	{
+		bIsGregorian = false;
+	}
+
+	if (Month < 3)
+	{
+		Year = Year - 1;
+		Month = Month + 12;
+	}
+
+	int32 b = 0;
+	if (bIsGregorian)
+	{
+		int32 a = FloorForJulianDate(Year / 100.0);
+		b = 2 - a + FloorForJulianDate(a / 4.0);
+	}
+
+	// Compute the Julian Date
+	double JulianDate = FloorForJulianDate(365.25 * (Year + 4716)) + FloorForJulianDate(30.6001 * (Month + 1)) + Day + b - 1524.5;
+	JulianDate += Hours / 24.0;
+	JulianDate += Minutes / 24.0 / 60.0;
+	JulianDate += Seconds / 24.0 / 60.0 / 60.0;
+	return JulianDate;
+}	
+
+FDateTime UCelestialMaths::JulianDateToUTCDateTime(double JulianDate)
+{
+	// From https://www.celestialprogramming.com/julian.html
+	// From Meeus, CH7, p63
+
+	double temp = JulianDate + 0.5;
+	int32 IntegralPart = StaticCast<int32>(FMath::TruncToDouble(temp));
+	double FractionalPart = temp - IntegralPart;
+
+	// If Integral Part < 2299161, take A = IntegralPart
+	int32 A = IntegralPart;
+	if (IntegralPart >= 2299161)
+	{
+		int32 Alpha = FloorForJulianDate((IntegralPart - 1867216.25) / 36524.25);
+		A = IntegralPart + 1 + Alpha - FloorForJulianDate(Alpha / 4);
+	}
+
+	// Compute ABCDE values
+	int32 B = A + 1524;
+	int32 C = FloorForJulianDate((B - 122.1) / 365.25);
+	int32 D = FloorForJulianDate(365.25 * C);
+	int32 E = FloorForJulianDate((B - D) / 30.6001);
+
+	int32 day=B-D-static_cast<int32>(30.6001 * E)+FractionalPart;
+	int32 month=E-1;
+	if(E>13){
+		month=E-13;
+	}
+	int32 year=C-4716;
+	if(month<3){
+		year=C-4715;
+	}
+
+	// Split in H,M,S,MS
+	int32 hour = StaticCast<int32>(FMath::TruncToDouble(FractionalPart * 24.0));
+	FractionalPart -= hour / 24.0;
+	int32 minute = StaticCast<int32>(FMath::TruncToDouble(FractionalPart * 24.0 * 60.0));
+	FractionalPart -= minute / (24.0 * 60.0);
+	int32 seconds = StaticCast<int32>(FMath::TruncToDouble(FractionalPart * 24.0 * 60.0 * 60.0));
+	FractionalPart -= seconds / (24.0 * 60.0 * 60.0);
+	int32 milliseconds = StaticCast<int32>(FMath::RoundToDouble(FractionalPart * 24.0 * 60.0 * 60.0 * 1000.0));
+
+	return FDateTime(year, month, day, hour, minute, seconds, milliseconds);
+}
+
+double UCelestialMaths::DateTimeToGreenwichMeanSiderealTime(FDateTime UTCDateTime)
+{
+	double JulianDate = UTCDateTime.GetJulianDay();
+	double JulianCentury = JulianDateToJulianCenturies(JulianDate);
+	double SideralTimeAt0h = 100.46061837 + 36000.770053608 * JulianCentury + 0.000387933 * JulianCentury * JulianCentury + 1.0/38710000.0 * JulianCentury * JulianCentury * JulianCentury;
+	return ModPositive(SideralTimeAt0h, 360);
+}
+
+double UCelestialMaths::JulianDateToGreenwichMeanSiderealTime(double JulianDate) {
+	//The IAU Resolutions on Astronomical Reference Systems, Time Scales, and Earth Rotation Models Explanation and Implementation (George H. Kaplan)
+	// https://arxiv.org/pdf/astro-ph/0602086.pdf - page 30
+	// T is the number of centuries of TDB (or TT) from J2000.0
+	// Formula for arcseconds: EarthRotationAngle + 0.014506 + 4612.15739966*T + 1.39667721*T^2 − 0.00009344*T^3 + 0.00001882*T^4
+	// Here our Earth angle is already in degrees
+	const double JulianCentury = JulianDateToJulianCenturies(JulianDate);
+
+	const double EarthRotationAngle = GetEarthRotationAngle(JulianDate);
+	
+	double GMST = EarthRotationAngle + ArcsecondsToDegrees(
+				0.014506 +
+				4612.15739966 * JulianCentury +
+				1.39667721 * JulianCentury * JulianCentury +
+				-0.00009344 * JulianCentury * JulianCentury * JulianCentury +
+				0.00001882 * JulianCentury * JulianCentury * JulianCentury * JulianCentury) ;
+
+	return ModPositive(GMST, 360.0);
+}
+
+double UCelestialMaths::LocalSideralTime(double LongitudeDegrees, double GreenwichMeanSideralTime)
+{
+	return GreenwichMeanSideralTime + LongitudeDegrees;
+}
+
+double UCelestialMaths::JulianDateToGreenwichApparentSiderealTime(double JulianDate)
+{
+	// From https://aa.usno.navy.mil/faq/GAST
+	// The Greenwich apparent sidereal time is obtained by adding a correction to the Greenwich mean sidereal time computed above.
+	// The correction term is called the nutation in right ascension or the equation of the equinoxes. Thus,
+	// GAST = GMST + eqeq.
+	
+	const double GMST = JulianDateToGreenwichMeanSiderealTime(JulianDate);
+	const double EE = EquationOfTheEquinoxes(JulianDate);
+	double gast = GMST + EE;
+
+	return ModPositive(gast,360);
+}
+
+double UCelestialMaths::JulianDateToJulianCenturies(double JulianDate)
+{
+	return (JulianDate - 2451545.0) / 36525.0;
+}
+
+double UCelestialMaths::GetLeapSeconds(double JulianDate)
+	{
+		//Source IERS Resolution B1 and http://maia.usno.navy.mil/ser7/tai-utc.dat
+		//This function must be updated any time a new leap second is introduced
+
+		if (JulianDate > 2457754.5) return 37.0;
+		if (JulianDate > 2457204.5) return 36.0;
+		if (JulianDate > 2456109.5) return 35.0;
+		if (JulianDate > 2454832.5) return 34.0;
+		if (JulianDate > 2453736.5) return 33.0;
+		if (JulianDate > 2451179.5) return 32.0;
+		if (JulianDate > 2450630.5) return 31.0;
+		if (JulianDate > 2450083.5) return 30.0;
+		if (JulianDate > 2449534.5) return 29.0;
+		if (JulianDate > 2449169.5) return 28.0;
+		if (JulianDate > 2448804.5) return 27.0;
+		if (JulianDate > 2448257.5) return 26.0;
+		if (JulianDate > 2447892.5) return 25.0;
+		if (JulianDate > 2447161.5) return 24.0;
+		if (JulianDate > 2446247.5) return 23.0;
+		if (JulianDate > 2445516.5) return 22.0;
+		if (JulianDate > 2445151.5) return 21.0;
+		if (JulianDate > 2444786.5) return 20.0;
+		if (JulianDate > 2444239.5) return 19.0;
+		if (JulianDate > 2443874.5) return 18.0;
+		if (JulianDate > 2443509.5) return 17.0;
+		if (JulianDate > 2443144.5) return 16.0;
+		if (JulianDate > 2442778.5) return 15.0;
+		if (JulianDate > 2442413.5) return 14.0;
+		if (JulianDate > 2442048.5) return 13.0;
+		if (JulianDate > 2441683.5) return 12.0;
+		if (JulianDate > 2441499.5) return 11.0;
+		if (JulianDate > 2441317.5) return 10.0;
+		if (JulianDate > 2439887.5) return 4.21317 + (JulianDate - 2439126.5) * 0.002592;
+		if (JulianDate > 2439126.5) return 4.31317 + (JulianDate - 2439126.5) * 0.002592;
+		if (JulianDate > 2439004.5) return 3.84013 + (JulianDate - 2438761.5) * 0.001296;
+		if (JulianDate > 2438942.5) return 3.74013 + (JulianDate - 2438761.5) * 0.001296;
+		if (JulianDate > 2438820.5) return 3.64013 + (JulianDate - 2438761.5) * 0.001296;
+		if (JulianDate > 2438761.5) return 3.54013 + (JulianDate - 2438761.5) * 0.001296;
+		if (JulianDate > 2438639.5) return 3.44013 + (JulianDate - 2438761.5) * 0.001296;
+		if (JulianDate > 2438486.5) return 3.34013 + (JulianDate - 2438761.5) * 0.001296;
+		if (JulianDate > 2438395.5) return 3.24013 + (JulianDate - 2438761.5) * 0.001296;
+		if (JulianDate > 2438334.5) return 1.945858 + (JulianDate - 2437665.5) * 0.0011232;
+		if (JulianDate > 2437665.5) return 1.845858 + (JulianDate - 2437665.5) * 0.0011232;
+		if (JulianDate > 2437512.5) return 1.372818 + (JulianDate - 2437300.5) * 0.001296;
+		if (JulianDate > 2437300.5) return 1.422818 + (JulianDate - 2437300.5) * 0.001296;
+		return 0.0;
+	}
+
+double UCelestialMaths::InternationalAtomicTimeToTerrestrialTime(double TAI)
+{
+	// From https://www2.mps.mpg.de/homes/fraenz/systems/systems2art/node2.html
+	//  TT = terrestrial time in SI seconds
+	//  TT= TAI + 32.184 seconds;
+	return TAI + 32.184;
+}
+
+double UCelestialMaths::JulianDateToInternationalAtomicTime(double JulianDate)
+{
+	return GetLeapSeconds(JulianDate) + DaysToSeconds(JulianDate);
+}
+
+#pragma endregion
+
+#pragma region Angles
+
+double UCelestialMaths::ModPositive(double Value, double Modulo)
+{
+	double Result = FMath::Fmod(Value, Modulo);
+	if (Result < 0)
+	{
+		Result += Modulo;
+	}
+	return Result; 
+}
+
+void UCelestialMaths::DegreesToHMS(double DecimalDegrees, int32& Hours, int32& Minutes, double& Seconds)
+{
+	DecimalDegrees = FMath::Fmod(DecimalDegrees, 360);
+	if (DecimalDegrees < 0)
+	{
+		DecimalDegrees += 360;
+	}
+
+	double AngleHours = DecimalDegrees/15.0;
+
+	Hours = StaticCast<uint32>(AngleHours);
+	Minutes = StaticCast<uint32>((AngleHours - Hours) * 60.0);
+	Seconds = (AngleHours - Hours) * 3600.0 - Minutes * 60.0;
+}
+
+void UCelestialMaths::DegreesToDMS(double DecimalDegrees, bool& Sign, int32& Degrees, int32& Minutes, double& Seconds)
+{
+	Sign = true;
+	DecimalDegrees = FMath::Fmod(DecimalDegrees, 360);
+
+	if (DecimalDegrees < 0)
+	{
+		Sign = false;
+		DecimalDegrees *= -1;
+	}
+	
+	Degrees = StaticCast<uint32>(DecimalDegrees);
+	Minutes = StaticCast<uint32>((DecimalDegrees - Degrees) * 60.0);
+	Seconds = (DecimalDegrees - Degrees) * 3600.0 - Minutes * 60.0;
+}
+
+#pragma endregion
+
+#pragma region Earth
+
+double UCelestialMaths::GetEarthRotationAngle(double JulianDate)
+{
+	//The IAU Resolutions on Astronomical Reference Systems, Time Scales, and Earth Rotation Models Explanation and Implementation (George H. Kaplan)
+	// https://arxiv.org/pdf/astro-ph/0602086.pdf - page 30, Eq 2.11
+	//
+	// DU is the number of UT1 days from 2000 January 1, 12h UT1: DU = JD(UT1)– 2451545.0.
+	// The angle θ is given in terms of rotations (units of 2π radians or 360d)
+	// 
+	// θ =0.7790572732640 +0.00273781191135448 DU + frac(JD(UT1)) --> We need to multiply it by 2xPI
+
+	const double DU = JulianDate - 2451545.0;
+	const double JulianDateFraction = FMath::Fractional(JulianDate);
+	const double Rotations = (0.779057273264 + 0.00273781191135448 * DU + JulianDateFraction);
+
+	return ModPositive(Rotations * 360, 360);
+}
+
+FVector UCelestialMaths::GeodeticLatLonToECEFXYZAU(double Latitude, double Longitude, double Altitude)
+{
+	// Algorithm from Explanatory Supplement to the Astronomical Almanac 3rd ed. P294
+	double LatitudeRadians = FMath::DegreesToRadians(Latitude);
+	double LongitudeRadians = FMath::DegreesToRadians(Longitude);
+	
+	const double a = MetersToAstronomicalUnits(6378136.6);
+	const double f = 1 / 298.25642;
+
+	const double C = 1 / FMath::Sqrt(FMath::Cos(LatitudeRadians) * FMath::Cos(LatitudeRadians) + (1.0 - f) * (1.0 - f) * (FMath::Sin(LatitudeRadians) * FMath::Sin(LatitudeRadians)));
+
+	const double S = (1 - f) * (1 - f) * C;
+	const double h = MetersToAstronomicalUnits(Altitude);
+
+	return FVector(
+		(a * C + h) * FMath::Cos(LatitudeRadians) * FMath::Cos(LongitudeRadians),
+		(a * C + h) * FMath::Cos(LatitudeRadians) * FMath::Sin(LongitudeRadians),
+		(a * S + h) * FMath::Sin(LatitudeRadians)); 
+}
+
+FVector UCelestialMaths::GetObserverGeocentricLocationAU(double Latitude,double Longitude, double Altitude, double JulianDate )
+{
+	const FVector ObserverECEF = GeodeticLatLonToECEFXYZAU(Latitude, Longitude, Altitude);
+	const double GreenwichApparentSiderealTime = JulianDateToGreenwichApparentSiderealTime(JulianDate);
+
+	// Compute cosine and sine of the angle
+	float CosTheta = FMath::Cos(-FMath::DegreesToRadians(GreenwichApparentSiderealTime));
+	float SinTheta = FMath::Sin(-FMath::DegreesToRadians(GreenwichApparentSiderealTime));
+	
+	// Construct the 3x3 rotation matrix to rotate the ECEF position around the Earth axis, depending on the GAST
+	FMatrix RotationMatrix = FMatrix(
+			FPlane(CosTheta, -SinTheta, 0.0f, 0.0f), // Row 1
+			FPlane(SinTheta,  CosTheta, 0.0f, 0.0f), // Row 2
+			FPlane(0.0f,      0.0f,     1.0f, 0.0f), // Row 3
+			FPlane(0.0f,      0.0f,     0.0f, 1.0f)  // Row 4 (required for FMatrix, but ignored here)
+		);
+	
+	return RotationMatrix.TransformVector(ObserverECEF);
+}
+
+double UCelestialMaths::EquationOfTheEquinoxes(double JulianDate) {
+	// The IAU Resolutions on Astronomical Reference Systems, Time Scales, and Earth Rotation Models Explanation and Implementation (George H. Kaplan)
+	// https://arxiv.org/pdf/astro-ph/0602086.pdf
+	// eq 5.12 p58
+	// The expression for the mean obliquity of the ecliptic (the angle between the mean equator and ecliptic, or, equivalently, between the ecliptic pole and mean celestial pole of date) is:
+	// E =E0 −46.836769T −0.0001831T2 +0.00200340T3 −0.000000576T4 −0.0000000434T5
+	// with E0 = 84381.406 arcseconds
+	// With T = the number of Julian centuries of TDB since 2000 Jan 1, 12h TDB. If the dates and times are expressed as Julian dates, then T = (t − 2451545.0)/36525.
+	const double JulianCenturies = JulianDateToJulianCenturies(JulianDate);
+
+	const double MeanObliquity = ArcsecondsToDegrees(
+			84381.406 +
+			-46.836769 * JulianCenturies +
+			-0.0001831 * JulianCenturies * JulianCenturies +
+			0.0020034 * JulianCenturies * JulianCenturies * JulianCenturies +
+			-0.000000576 * JulianCenturies * JulianCenturies * JulianCenturies * JulianCenturies +
+			-0.0000000434 * JulianCenturies * JulianCenturies * JulianCenturies * JulianCenturies * JulianCenturies);
+
+	double DeltaPsi;
+	double DeltaEpsilon; 
+	Nutation2000BTruncated(JulianDate, DeltaPsi, DeltaEpsilon);
+	return DeltaPsi * FMath::Cos(FMath::DegreesToRadians(MeanObliquity + DeltaEpsilon));
+}
+
+void UCelestialMaths::Nutation2000BTruncated(double JulianDate, double& DeltaPsi, double& DeltaEpsilon)
+{
+	//The IAU Resolutions on Astronomical Reference Systems, Time Scales, and Earth Rotation Models Explanation and Implementation (George H. Kaplan)
+	//https://arxiv.org/pdf/astro-ph/0602086.pdf
+	//IAU 2000B Nutation truncated to 6 terms
+
+	const double JC = JulianDateToJulianCenturies(JulianDate);
+	const double JC2 = JC * JC;
+	const double JC3 = JC * JC2;
+	const double JC4 = JC * JC3;
+
+	//Fundamental Arguments p46 eq 5.17, 5.18, 5.19
+	// The last five arguments are the same fundamental luni-solar arguments used in previous nutation theories, but with updated expresssions. They are, respectively,
+	// l the mean anomaly of the Moon; 
+	// l′the mean anomaly of the Sun; --> Lp
+	// F the mean argument of latitude of the Moon; 
+	// D the mean elongation of the Moon from the Sun
+	// Ω the mean longitude of the Moon’s mean ascending node:
+
+	const double Lp = ArcsecondsToRadians(1287104.79305 + 129596581.0481 * JC - 0.5532 * JC2 + 0.000136 * JC3 - 0.00001149 * JC4);
+	const double F = ArcsecondsToRadians(335779.526232 + 1739527262.8478 * JC - 12.7512 * JC2 - 0.001037 * JC3 + 0.00000417 * JC4);
+	const double D = ArcsecondsToRadians(1072260.70369 + 1602961601.209 * JC - 6.3706 * JC2 + 0.006593 * JC3 - 0.00003169 * JC4);
+	const double Omega = ArcsecondsToRadians(450160.398036 - 6962890.5431 * JC + 7.4722 * JC2 + 0.007702 * JC3 - 0.00005939 * JC4);
+
+	//Terms summed from lowest to highest to reduce floating point rounding errors.  See coefficients Page 88.
+	// Constants are first multiplied by 10000000 to reduce errors
+	double DeltaPsiArcseconds = 0;
+	double DeltaEpsilonArcSeconds = 0;
+
+	// FundamentalArgument #6
+	double FundamentalArgument = Lp + 2 * (F - D + Omega);
+	DeltaPsiArcseconds += (-516821 + 1226 * JC) * FMath::Sin(FundamentalArgument) + -524 * FMath::Cos(FundamentalArgument);
+	DeltaEpsilonArcSeconds += (224386 + -677 * JC) * FMath::Cos(FundamentalArgument) + -174 * FMath::Sin(FundamentalArgument);
+
+	// FundamentalArgument #5 - Just Lp
+	DeltaPsiArcseconds += (1475877 + -3633 * JC) * FMath::Sin(Lp) + 11817 * FMath::Cos(Lp);
+	DeltaEpsilonArcSeconds += (73871 + -184 * JC) * FMath::Cos(Lp) + -1924 * FMath::Sin(Lp);
+
+	// FundamentalArgument #4
+	FundamentalArgument = 2 * Omega;
+	DeltaPsiArcseconds += (2074554 + 207 * JC) * FMath::Sin(FundamentalArgument) + -698 * FMath::Cos(FundamentalArgument);
+	DeltaEpsilonArcSeconds += (-897492 + 470 * JC) * FMath::Cos(FundamentalArgument) + -291 * FMath::Sin(FundamentalArgument);
+
+	// FundamentalArgument #3
+	FundamentalArgument = 2 * (F + Omega);
+	DeltaPsiArcseconds += (-2276413 + -234 * JC) * FMath::Sin(FundamentalArgument) + 2796 * FMath::Cos(FundamentalArgument);
+	DeltaEpsilonArcSeconds += (978459 + -485 * JC) * FMath::Cos(FundamentalArgument) + 1374 * FMath::Sin(FundamentalArgument);
+
+	// FundamentalArgument #2
+	FundamentalArgument = 2 * (F - D + Omega);
+	DeltaPsiArcseconds += (-13170906 + -1675 * JC) * FMath::Sin(FundamentalArgument) + -13696 * FMath::Cos(FundamentalArgument);
+	DeltaEpsilonArcSeconds += (5730336 + -3015 * JC) * FMath::Cos(FundamentalArgument) + -4587 * FMath::Sin(FundamentalArgument);
+
+	// FundamentalArgument #1
+	DeltaPsiArcseconds += (-172064161 + -174666 * JC) * FMath::Sin(Omega) + 33386 * FMath::Cos(Omega);
+	DeltaEpsilonArcSeconds += (92052331 + 9086 * JC) * FMath::Cos(Omega) + 15377 * FMath::Sin(Omega);
+
+	DeltaPsi = ArcsecondsToDegrees(DeltaPsiArcseconds / 10000000);
+	DeltaEpsilon = ArcsecondsToDegrees(DeltaEpsilonArcSeconds / 10000000);
+}
+
+FTransform UCelestialMaths::GetPlanetCenterTransform(double Latitude, double Longitude, double Altitude)
+{
+	// Compute the Location part 
+	FVector ECEFLocation = GeodeticLatLonToECEFXYZAU(Latitude, Longitude, Altitude) * AstronomicalUnitsMeters;
+
+	// Compute the 3 Axis vectors
+	FMatrix AxisMatrix; 
+	// See ECEF standard : https://commons.wikimedia.org/wiki/File:ECEF_ENU_Longitude_Latitude_right-hand-rule.svg
+	if (FMathd::Abs(ECEFLocation.X) < FMathd::Epsilon &&
+		FMathd::Abs(ECEFLocation.Y) < FMathd::Epsilon)
+	{
+		// Special Case - On earth axis... 
+		double Sign = 1.0;
+		if (FMathd::Abs(ECEFLocation.Z) < FMathd::Epsilon)
+		{
+			// At origin - Should not happen, but consider it's the same as north pole
+			// Leave Sign = 1
+		}
+		else
+		{
+			// At South or North pole - Axis are set to be continuous with other points
+			Sign = FMathd::SignNonZero(ECEFLocation.Z);
+		}
+
+		AxisMatrix = FMatrix(
+			FVector::YAxisVector, 			// East = Y
+			-FVector::XAxisVector * Sign,	// North = Sign * X
+			FVector::ZAxisVector * Sign,	// Up = Sign*Z
+			ECEFLocation);
+	}
+	else
+	{
+		double Tolerance = 1.e-50; // Normalize with a very low threshold, because default is 10-8, too high for double computations
+
+		// Compute the ellipsoid normal (Earth...)
+		double SemiMajorMeter = 6378137.0;
+		double SemiMinorMetre = 6356752.3142451793;
+		
+		FVector OneOverRadiiSquared = FVector(1.0 / (SemiMajorMeter * SemiMajorMeter), 1.0 / (SemiMajorMeter * SemiMajorMeter), 1.0 / (SemiMinorMetre * SemiMinorMetre)); 
+		FVector GeodeticSurfaceNormal( ECEFLocation.X * OneOverRadiiSquared.X, ECEFLocation.Y * OneOverRadiiSquared.Y, ECEFLocation.Z * OneOverRadiiSquared.Z);
+		GeodeticSurfaceNormal.Normalize(Tolerance);
+
+		// Get other axes
+		FVector Up = GeodeticSurfaceNormal;
+		FVector East(-ECEFLocation.Y, ECEFLocation.X, 0.0); 
+		East.Normalize(Tolerance); 
+		FVector North = Up.Cross(East);
+
+		// Set Matrix
+		AxisMatrix = FMatrix(	East, North, Up, ECEFLocation);
+	}
+	
+	return FTransform(AxisMatrix.Inverse());
+}
+
+#pragma endregion
+
+#pragma region Sun
+
+FSunInfo UCelestialMaths::GetSunInformation(double JulianDate, double ObserverLatitude, double ObserverLongitude)
+{
+    	FSunInfo SunInfo;
+
+		// TODO_Beta
+		// This code comes from the former SunPosition Plugin. Add the missing properties but it will require Time Zone, DST, and such... 
+	
+
+	
+		// if (!FDateTime::Validate(Year, Month, Day, Hours, Minutes, Seconds, Milliseconds))
+		// {
+		// 	return SunInfo;
+		// }
+		//
+		// FDateTime CalcTime(Year, Month, Day, Hours, Minutes, Seconds, Milliseconds);
+		//
+		// double TimeOffset = TimeZone;
+		// if (bIsDaylightSavingTime)
+		// {
+		// 	TimeOffset += 1.0;
+		// }
+
+		double LatitudeRad = FMath::DegreesToRadians(ObserverLatitude);
+
+		// Get the julian day (number of days since Jan 1st of the year 4713 BC)
+		double JulianDay = JulianDate;
+		double JulianCentury = JulianDateToJulianCenturies(JulianDate);
+
+		// Get the sun's mean longitude , referred to the mean equinox of julian date
+		double GeomMeanLongSunDeg = FMath::Fmod(280.46646 + JulianCentury*(36000.76983 + JulianCentury*0.0003032), 360.0);
+		double GeomMeanLongSunRad = FMath::DegreesToRadians(GeomMeanLongSunDeg);
+
+		// Get the sun's mean anomaly
+		double GeomMeanAnomSunDeg = 357.52911 + JulianCentury*(35999.05029 - 0.0001537*JulianCentury);
+		double GeomMeanAnomSunRad = FMath::DegreesToRadians(GeomMeanAnomSunDeg);
+
+		// Get the earth's orbit eccentricity
+		double EccentEarthOrbit = 0.016708634 - JulianCentury*(0.000042037 + 0.0000001267*JulianCentury);
+		
+		// Get the sun's equation of the center
+		double SunEqOfCtr = FMath::Sin(GeomMeanAnomSunRad)*(1.914602 - JulianCentury*(0.004817 + 0.000014*JulianCentury))
+						  + FMath::Sin(2.0 * GeomMeanAnomSunRad)*(0.019993 - 0.000101*JulianCentury)
+						  + FMath::Sin(3.0 * GeomMeanAnomSunRad)*0.000289;
+		
+		// Get the sun's true longitude
+		double SunTrueLongDeg = GeomMeanLongSunDeg + SunEqOfCtr;
+
+		// Get the sun's true anomaly
+		//	double SunTrueAnomDeg = GeomMeanAnomSunDeg + SunEqOfCtr;
+		//	double SunTrueAnomRad = FMath::DegreesToRadians(SunTrueAnomDeg);
+
+		// Get the earth's distance from the sun
+		//	double SunRadVectorAUs = (1.000001018*(1.0 - EccentEarthOrbit*EccentEarthOrbit)) / (1.0 + EccentEarthOrbit*FMath::Cos(SunTrueAnomRad));
+		
+		// Get the sun's apparent longitude
+		double SunAppLongDeg = SunTrueLongDeg - 0.00569 - 0.00478*FMath::Sin(FMath::DegreesToRadians(125.04 - 1934.136*JulianCentury));
+		double SunAppLongRad = FMath::DegreesToRadians(SunAppLongDeg);
+
+		// Get the earth's mean obliquity of the ecliptic
+		double MeanObliqEclipticDeg = 23.0 + (26.0 + ((21.448 - JulianCentury*(46.815 + JulianCentury*(0.00059 - JulianCentury*0.001813)))) / 60.0) / 60.0;
+		
+		// Get the oblique correction
+		double ObliqCorrDeg = MeanObliqEclipticDeg + 0.00256*FMath::Cos(FMath::DegreesToRadians(125.04 - 1934.136*JulianCentury));
+		double ObliqCorrRad = FMath::DegreesToRadians(ObliqCorrDeg);
+
+		// Get the sun's right ascension
+		double SunRtAscenRad = FMath::Atan2(FMath::Cos(ObliqCorrRad)*FMath::Sin(SunAppLongRad), FMath::Cos(SunAppLongRad));
+		double SunRtAscenDeg = FMath::RadiansToDegrees(SunRtAscenRad);
+		
+		// Get the sun's declination
+		double SunDeclinRad = FMath::Asin(FMath::Sin(ObliqCorrRad)*FMath::Sin(SunAppLongRad));
+		double SunDeclinDeg = FMath::RadiansToDegrees(SunDeclinRad);
+
+		double VarY = FMath::Pow(FMath::Tan(ObliqCorrRad / 2.0), 2.0);
+		
+		// Get the equation of time
+		double EqOfTimeMinutes = 4.0 * FMath::RadiansToDegrees(VarY*FMath::Sin(2.0 * GeomMeanLongSunRad) - 2.0 * EccentEarthOrbit*FMath::Sin(GeomMeanAnomSunRad) + 4.0 * EccentEarthOrbit*VarY*FMath::Sin(GeomMeanAnomSunRad)*FMath::Cos(2.0 * GeomMeanLongSunRad) - 0.5*VarY*VarY*FMath::Sin(4.0 * GeomMeanLongSunRad) - 1.25*EccentEarthOrbit*EccentEarthOrbit*FMath::Sin(2.0 * GeomMeanAnomSunRad));
+
+		// Get the hour angle of the sunrise
+		double HASunriseDeg = FMath::RadiansToDegrees(FMath::Acos(FMath::Cos(FMath::DegreesToRadians(90.833)) / (FMath::Cos(LatitudeRad)*FMath::Cos(SunDeclinRad)) - FMath::Tan(LatitudeRad)*FMath::Tan(SunDeclinRad)));
+		//	double SunlightDurationMinutes = 8.0 * HASunriseDeg;
+
+		// // Get the local time of the sun's rise and set
+		// double SolarNoonLST = (720.0 - 4.0 * ObserverLongitude - EqOfTimeMinutes + TimeOffset * 60.0) / 1440.0;
+		// double SunriseTimeLST = SolarNoonLST - HASunriseDeg * 4.0 / 1440.0;
+		// double SunsetTimeLST = SolarNoonLST + HASunriseDeg * 4.0 / 1440.0;
+		//
+		// // Get the true solar time
+		// double TrueSolarTimeMinutes = FMath::Fmod(CalcTime.GetTimeOfDay().GetTotalMinutes() + EqOfTimeMinutes + 4.0 * ObserverLongitude - 60.0 * TimeOffset, 1440.0);
+		//
+		// // Get the hour angle of current time
+		// double HourAngleDeg = TrueSolarTimeMinutes < 0 ? TrueSolarTimeMinutes / 4.0 + 180 : TrueSolarTimeMinutes / 4.0 - 180.0;
+		// double HourAngleRad = FMath::DegreesToRadians(HourAngleDeg);
+		//
+		// // Get the solar zenith angle
+		// double SolarZenithAngleRad = FMath::Acos(FMath::Sin(LatitudeRad)*FMath::Sin(SunDeclinRad) + FMath::Cos(LatitudeRad)*FMath::Cos(SunDeclinRad)*FMath::Cos(HourAngleRad));
+		// double SolarZenithAngleDeg = FMath::RadiansToDegrees(SolarZenithAngleRad);
+		//
+		// // Get the sun elevation
+		// double SolarElevationAngleDeg = 90.0 - SolarZenithAngleDeg;
+		// double SolarElevationAngleRad = FMath::DegreesToRadians(SolarElevationAngleDeg);
+		// double TanOfSolarElevationAngle = FMath::Tan(SolarElevationAngleRad);
+		//
+		// // Get the approximated atmospheric refraction
+		// double ApproxAtmosphericRefractionDeg = 0.0;
+		// if (SolarElevationAngleDeg <= 85.0)
+		// {
+		// 	if (SolarElevationAngleDeg > 5.0)
+		// 	{
+		// 		ApproxAtmosphericRefractionDeg = 58.1 / TanOfSolarElevationAngle - 0.07 / FMath::Pow(TanOfSolarElevationAngle, 3) + 0.000086 / FMath::Pow(TanOfSolarElevationAngle, 5) / 3600.0;
+		// 	}
+		// 	else
+		// 	{
+		// 		if (SolarElevationAngleDeg > -0.575)
+		// 		{
+		// 			ApproxAtmosphericRefractionDeg = 1735.0 + SolarElevationAngleDeg * (-518.2 + SolarElevationAngleDeg * (103.4 + SolarElevationAngleDeg * (-12.79 + SolarElevationAngleDeg * 0.711)));
+		// 		}
+		// 		else
+		// 		{
+		// 			ApproxAtmosphericRefractionDeg = -20.772 / TanOfSolarElevationAngle;
+		// 		}
+		// 	}
+		// 	ApproxAtmosphericRefractionDeg /= 3600.0;
+		// }
+		//◘
+		// // Get the corrected solar elevation
+		// double SolarElevationcorrectedforatmrefractionDeg = SolarElevationAngleDeg + ApproxAtmosphericRefractionDeg;
+		//
+		// // Get the solar azimuth 
+		// double tmp = FMath::RadiansToDegrees(FMath::Acos(((FMath::Sin(LatitudeRad)*FMath::Cos(SolarZenithAngleRad)) - FMath::Sin(SunDeclinRad)) / (FMath::Cos(LatitudeRad)*FMath::Sin(SolarZenithAngleRad))));
+		// double SolarAzimuthAngleDegcwfromN = HourAngleDeg > 0.0 ? FMath::Fmod(tmp + 180.0, 360.0) : FMath::Fmod(540.0 - tmp, 360.0);
+
+
+		// offset elevation angle to fit with UE coords system
+    	SunInfo.RA = SunRtAscenDeg / 15.0;
+		SunInfo.DEC = SunDeclinDeg; 
+		// SunInfo.Elevation = 180.0 + SolarElevationAngleDeg;
+		// SunInfo.CorrectedElevation = 180.0 + SolarElevationcorrectedforatmrefractionDeg;
+		// SunInfo.Azimuth = SolarAzimuthAngleDegcwfromN;
+		// SunInfo.SolarNoon = FTimespan::FromDays(SolarNoonLST);
+		// SunInfo.SunriseTime = FTimespan::FromDays(SunriseTimeLST);
+		// SunInfo.SunsetTime = FTimespan::FromDays(SunsetTimeLST);
+    	return SunInfo;
+	}
+
+#pragma endregion
+
+#pragma region Utilities
+
+void UCelestialMaths::XYZToRADEC_RH(FVector XYZ, double& RADegrees, double& DECDegrees, double& Radius)
+{
+	//Convert from Cartesian to polar coordinates
+	Radius = XYZ.Length();
+	double l = FMath::Atan2(XYZ.Y, XYZ.X);
+	double t = FMath::Acos(XYZ.Z / Radius);
+
+	//Make sure RA is positive, and Dec is in range +/-90
+	if (l < 0) {
+		l += 2 * UE_PI ;
+	}
+	t = 0.5 * UE_PI - t;
+
+	RADegrees = FMath::RadiansToDegrees(l);
+	DECDegrees = FMath::RadiansToDegrees(t);
+}
+
+FVector UCelestialMaths::RADECToXYZ_RH(double RADegrees, double DECDegrees, double Radius)
+{
+	double RARadians = FMath::DegreesToRadians(RADegrees);
+	double DECRadians = FMath::DegreesToRadians(DECDegrees);
+	
+	double X = Radius * FMath::Cos(DECRadians) * FMath::Cos(RARadians);
+	double Y = Radius * FMath::Cos(DECRadians) * FMath::Sin(RARadians);
+	double Z = Radius * FMath::Sin(DECRadians);
+
+	return FVector(X, Y, Z);
+}
+
+FString UCelestialMaths::GetPreciseVectorString(FVector Vector, int32 MinimumFractionalDigits)
+{
+	FNumberFormattingOptions NumberFormatOptions;
+	NumberFormatOptions.MinimumFractionalDigits = MinimumFractionalDigits;
+
+	return FString::Printf(TEXT("X=%s Y=%s Z=%s"),
+		*FText::AsNumber(Vector.X, &NumberFormatOptions).ToString(),
+		*FText::AsNumber(Vector.Y, &NumberFormatOptions).ToString(),
+		*FText::AsNumber(Vector.Z, &NumberFormatOptions).ToString()
+		);
+}
+
+FString UCelestialMaths::Conv_StarInfoToString(const FStarInfo& StarInfo)
+{
+	return StarInfo.ToString();
+}
+
+FString UCelestialMaths::Conv_PlanetaryBodyInfoToString(const FPlanetaryBodyInfo& PlanetaryBodyInfo)
+{
+	return PlanetaryBodyInfo.ToString();
+}
+
+FString UCelestialMaths::Conv_SunInfoToString(const FSunInfo& SunInfo)
+{
+	return SunInfo.ToString();
+}
+
+FString UCelestialMaths::Conv_RightAscensionToString(const double& RightAscensionHours)
+{
+	FString RAString;
+	int32 H = 0;
+	int32 M = 0;
+	double S = 0.0;
+	DegreesToHMS(RightAscensionHours * 15.0, H, M,S);
+	return Conv_HMSToString(H,M,S); 
+}
+
+FString UCelestialMaths::Conv_DeclinationToString(const double& DeclinationDegrees)
+{
+	int32 DecD = 0;
+	int32 DecM = 0;
+	double DecS = 0.0;
+	bool DecSign = false;
+	DegreesToDMS(DeclinationDegrees, DecSign, DecD, DecM, DecS  );
+	return Conv_DMSToString(DecSign, DecD, DecM, DecS);
+}
+
+FString UCelestialMaths::Conv_HMSToString( int32 Hours, int32 Minutes, double Seconds)
+{
+	return FString::Printf(TEXT("%dh%02dm%05.2fs"), Hours, Minutes, Seconds); 
+}
+
+FString UCelestialMaths::Conv_DMSToString(bool Sign, int32 Degrees, int32 Minutes, double Seconds)
+{
+	return FString::Printf(TEXT("%c%d°%02d\'%05.2f\""), Sign? '+' : '-' , Degrees, Minutes, Seconds); 
+}
+
+#pragma endregion
+
+// Private
+
+int32 UCelestialMaths::FloorForJulianDate(double JulianDate)
+{
+	if (JulianDate > 0)
+	{
+		return StaticCast<int32>(FMath::Floor(JulianDate));
+	}
+	if (JulianDate == FMath::Floor(JulianDate))
+	{
+		return StaticCast<int32>(JulianDate);
+	}
+	return StaticCast<int32>(FMath::Floor(JulianDate)-1);
+}
\ No newline at end of file
diff --git a/Engine/Plugins/Experimental/CelestialVault/Source/CelestialVault/Private/CelestialVaultDaySequenceActor.cpp b/Engine/Plugins/Experimental/CelestialVault/Source/CelestialVault/Private/CelestialVaultDaySequenceActor.cpp
new file mode 100644
--- /dev/null
+++ b/Engine/Plugins/Experimental/CelestialVault/Source/CelestialVault/Private/CelestialVaultDaySequenceActor.cpp
@@ -0,0 +1,748 @@
+﻿// Copyright Epic Games, Inc. All Rights Reserved.
+
+
+#include "CelestialVaultDaySequenceActor.h"
+
+
+// Components
+#include "Components/DirectionalLightComponent.h"
+#include "Components/ExponentialHeightFogComponent.h"
+#include "Components/InstancedStaticMeshComponent.h"
+#include "Components/SkyAtmosphereComponent.h"
+#include "Components/SkyLightComponent.h"
+#include "Components/VolumetricCloudComponent.h"
+#include "Components/StaticMeshComponent.h"
+
+// UE Objects
+#include "Materials/MaterialInterface.h"
+#include "UObject/ConstructorHelpers.h"
+#include "DaySequenceSubsystem.h"
+#include "DaySequenceCollectionAsset.h"
+#include "Engine/World.h"
+#include "Engine/StaticMesh.h"
+
+// Celestial Objects
+#include "CelestialMaths.h"
+#include "CelestialDataTypes.h"
+#include "CelestialVault.h"
+
+
+
+
+// Sets default values
+ACelestialVaultDaySequenceActor::ACelestialVaultDaySequenceActor(const FObjectInitializer& Init)
+: Super(Init)
+{
+	// ExponentialHeightFogComponent = CreateOptionalDefaultSubobject<UExponentialHeightFogComponent>(TEXT("ExponentialHeightFog"));
+	// ExponentialHeightFogComponent->SetupAttachment(RootComponent);
+	// ExponentialHeightFogComponent->bEnableVolumetricFog = true;
+
+	// Components Attached to Root
+	SkyLightComponent = CreateDefaultSubobject<USkyLightComponent>(TEXT("Sky Light"));
+	SkyLightComponent->SetupAttachment(RootComponent);
+	SkyLightComponent->bRealTimeCapture = true;
+	SkyLightComponent->bLowerHemisphereIsBlack = false;
+
+	VolumetricCloudComponent = CreateOptionalDefaultSubobject<UVolumetricCloudComponent>(TEXT("Volumetric Cloud"));
+	VolumetricCloudComponent->SetupAttachment(RootComponent);
+	if (!IsTemplate())
+	{
+		// We don't want to load this material for the CDO as it will hold on to it forever and it is quite a large asset.
+		static ConstructorHelpers::FObjectFinder<UMaterialInterface> VolumetricCloudDefaultMaterialRef(TEXT("/Engine/EngineSky/VolumetricClouds/m_SimpleVolumetricCloud_Inst.m_SimpleVolumetricCloud_Inst"));
+		VolumetricCloudComponent->SetMaterial(VolumetricCloudDefaultMaterialRef.Object);
+	}
+
+	PlanetCenterComponent = CreateDefaultSubobject<USceneComponent>(TEXT("Planet Center Transform"));
+	PlanetCenterComponent->SetupAttachment(RootComponent);
+
+	// Components Attached to Planet Center
+	SkyAtmosphereComponent = CreateDefaultSubobject<USkyAtmosphereComponent>(TEXT("Sky Atmosphere"));
+	SkyAtmosphereComponent->SetupAttachment(PlanetCenterComponent);
+	SkyAtmosphereComponent->TransformMode = ESkyAtmosphereTransformMode::PlanetCenterAtComponentTransform;
+
+
+	// Rotating Celestial Vault 
+	CelestialVaultComponent = CreateDefaultSubobject<USceneComponent>(TEXT("Rotating Celestial Vault"));
+	CelestialVaultComponent->SetupAttachment(PlanetCenterComponent);
+
+	// Components attached to the Celestial Vaul
+	
+	// Deep Sky background
+	static ConstructorHelpers::FObjectFinder<UStaticMesh> SkySphereDefaultMesh(TEXT("/CelestialVault/Meshes/CelestialVault.CelestialVault"));
+	static ConstructorHelpers::FObjectFinder<UMaterialInterface> SkySphereDefaultMaterial(TEXT("/CelestialVault/Materials/MI_CelestialVault.MI_CelestialVault"));
+	DeepSkyComponent = CreateOptionalDefaultSubobject<UStaticMeshComponent>(TEXT("Deep Sky"));
+	DeepSkyComponent->SetupAttachment(CelestialVaultComponent);
+	DeepSkyComponent->SetStaticMesh(SkySphereDefaultMesh.Object);
+	DeepSkyComponent->SetCollisionEnabled(ECollisionEnabled::NoCollision);
+	DeepSkyComponent->SetGenerateOverlapEvents(false);
+	DeepSkyComponent->SetCastShadow(false);
+	DeepSkyComponent->SetAffectDynamicIndirectLighting(false);
+	DeepSkyComponent->SetCanEverAffectNavigation(false);
+	DeepSkyComponent->SetMaterial(0, SkySphereDefaultMaterial.Object.Get());
+	DeepSkyComponent->SetRelativeScale3D(FVector(CelestialVaultDistance*1000.0));
+
+	// Stars ISM
+	static ConstructorHelpers::FObjectFinder<UStaticMesh> PlaneXMesh(TEXT("/CelestialVault/Meshes/PlaneX.PlaneX"));
+	static ConstructorHelpers::FObjectFinder<UMaterialInterface> StarsDefaultMaterial(TEXT("/CelestialVault/Materials/MI_Stars.MI_Stars"));
+	StarsComponent = CreateOptionalDefaultSubobject<UInstancedStaticMeshComponent>(TEXT("Stars"));
+	StarsComponent->SetupAttachment(CelestialVaultComponent);
+	StarsComponent->SetStaticMesh(PlaneXMesh.Object);
+	StarsComponent->SetCollisionEnabled(ECollisionEnabled::NoCollision);
+	StarsComponent->SetGenerateOverlapEvents(false);
+	StarsComponent->SetCastShadow(false);
+	StarsComponent->SetAffectDynamicIndirectLighting(false);
+	StarsComponent->SetCanEverAffectNavigation(false);
+	StarsComponent->SetMaterial(0, StarsDefaultMaterial.Object);
+
+	// Planets ISM
+	static ConstructorHelpers::FObjectFinder<UMaterialInterface> PlanetsDefaultMaterial(TEXT("/CelestialVault/Materials/MI_SolarSystemPlanets.MI_SolarSystemPlanets"));
+	PlanetsComponent = CreateOptionalDefaultSubobject<UInstancedStaticMeshComponent>(TEXT("Planets"));
+	PlanetsComponent->SetupAttachment(CelestialVaultComponent);
+	PlanetsComponent->SetStaticMesh(PlaneXMesh.Object);
+	PlanetsComponent->SetCollisionEnabled(ECollisionEnabled::NoCollision);
+	PlanetsComponent->SetGenerateOverlapEvents(false);
+	PlanetsComponent->SetCastShadow(false);
+	PlanetsComponent->SetAffectDynamicIndirectLighting(false);
+	PlanetsComponent->SetCanEverAffectNavigation(false);
+	PlanetsComponent->SetMaterial(0, PlanetsDefaultMaterial.Object);
+
+	// Moon - Disc
+	static ConstructorHelpers::FObjectFinder<UMaterialInterface> MoonDiscDefaultMaterial(TEXT("/CelestialVault/Materials/MI_Moon.MI_Moon"));
+	MoonDiscComponent = CreateOptionalDefaultSubobject<UStaticMeshComponent>(TEXT("Moon Disk"));
+	MoonDiscComponent->SetupAttachment(CelestialVaultComponent);
+	MoonDiscComponent->SetStaticMesh(PlaneXMesh.Object);
+	MoonDiscComponent->SetCollisionEnabled(ECollisionEnabled::NoCollision);
+	MoonDiscComponent->SetGenerateOverlapEvents(false);
+	MoonDiscComponent->SetCastShadow(true); // Eclipses?
+	MoonDiscComponent->SetAffectDynamicIndirectLighting(false);
+	MoonDiscComponent->SetCanEverAffectNavigation(false);
+	MoonDiscComponent->SetMaterial(0, MoonDiscDefaultMaterial.Object.Get());
+
+	// Moon - Light
+	MoonLightComponent = CreateDefaultSubobject<UDirectionalLightComponent>(TEXT("Moon Light"));
+	MoonLightComponent->SetupAttachment(CelestialVaultComponent);
+	MoonLightComponent->SetAtmosphereSunLightIndex(1);	// Make Moon the secondary directional light that contributes to the sky atmosphere.
+	MoonLightComponent->SetForwardShadingPriority(0); // Give Moon forward shading priority.
+	MoonLightComponent->SetIntensity(0.32f);
+	MoonLightComponent->SetUseTemperature(true);
+	MoonLightComponent->SetTemperature(9000.f);
+	MoonLightComponent->SetWorldRotation(FRotator(-45.f, 0.f, 0.f));
+	
+	// Attach the sunlight relative to the rotating Vault
+	SunLightComponent = CreateDefaultSubobject<UDirectionalLightComponent>(TEXT("Sun Light"));
+	SunLightComponent->SetupAttachment(CelestialVaultComponent);
+	SunLightComponent->SetAtmosphereSunLightIndex(0);	// Make Sun the first directional light that contributes to the sky atmosphere.
+	SunLightComponent->SetForwardShadingPriority(1);   // Give Sun forward shading priority.
+	SunLightComponent->SetIntensity(120000.0f);
+	
+	// Sequence and Data Assets
+	if (!IsTemplate())
+	{
+		// Override the default collection (which animates the moon and sky material)
+		static ConstructorHelpers::FObjectFinder<UDaySequenceCollectionAsset> DefaultCollection(TEXT("/CelestialVault/DSCA_CelestialVault.DSCA_CelestialVault"));
+		DaySequenceCollections.Add(DefaultCollection.Object.Get());
+
+		static ConstructorHelpers::FObjectFinder<UDataTable> DefaultStarsCatalog(TEXT("/CelestialVault/Data/HYGCatalog_10K.HYGCatalog_10K"));
+		CelestialStarCatalog = DefaultStarsCatalog.Object.Get();
+
+		static ConstructorHelpers::FObjectFinder<UDataTable> DefaultPlanetaryBodiesCatalog(TEXT("/CelestialVault/Data/SolarSystemPlanets.SolarSystemPlanets"));
+		PlanetsCatalog = DefaultPlanetaryBodiesCatalog.Object.Get();
+	}
+}
+
+//////// Protected Functions
+
+void ACelestialVaultDaySequenceActor::BeginPlay()
+{
+	Super::BeginPlay();
+
+	// We don't inherit from a BaseDaySequenceActor (not the same components), so we need to register ourselves
+	if (const UWorld* World = GetWorld())
+	{
+		if (UDaySequenceSubsystem* DaySequenceSubsystem = World->GetSubsystem<UDaySequenceSubsystem>())
+		{
+			DaySequenceSubsystem->SetDaySequenceActor(this);
+		}
+
+		//RebuildAll(); // TODO - is it needed? 
+	}
+}
+
+void ACelestialVaultDaySequenceActor::OnConstruction(const FTransform& Transform)
+{
+	Super::OnConstruction(Transform);
+
+	// We don't inherit from a BaseDaySequenceActor (not the same components), so we need to register ourselves
+	if (const UWorld* World = GetWorld())
+	{
+		if (UDaySequenceSubsystem* DaySequenceSubsystem = World->GetSubsystem<UDaySequenceSubsystem>())
+		{
+			if (DaySequenceSubsystem->GetDaySequenceActor(/* bFindFallbackOnNull */ false) != this)
+			{
+				DaySequenceSubsystem->SetDaySequenceActor(this);
+			}
+		}
+
+		RebuildAll();  
+	}
+}
+
+//////// Public Functions
+FDateTime ACelestialVaultDaySequenceActor::GetDate() const
+{
+	if (bUseCurrentDate)
+	{
+		return FDateTime::Now().GetDate();
+	}
+
+	if ( Day > FDateTime::DaysInMonth(Year, Month))
+	{
+		int32 ClampedDay = FDateTime::DaysInMonth(Year, Month);
+		UE_LOG(LogCelestialVault, Warning, TEXT("Day value (%d) over the number of days in month - Using %d instead"), Day, ClampedDay);
+		return FDateTime(Year, Month, ClampedDay, 0,0,0);
+	}
+		
+	return FDateTime(Year, Month, Day, 0,0,0);
+}
+
+FSunInfo ACelestialVaultDaySequenceActor::GetSunInfo(double JulianDate)
+{
+	FSunInfo ResultSunInfo = UCelestialMaths::GetSunInformation(JulianDate, Latitude, Longitude);
+
+	// Location
+	FVector SunLocation = UCelestialMaths::RADECToXYZ_RH(ResultSunInfo.RA * 15, ResultSunInfo.DEC, 1000); // TODO - Add the Sun Distance To SunInfo 
+	SunLocation *= FVector(100, -100, 100); // Convert to UE Frame by inverting Y and scaling to UE Units - // TODO - use AU
+	ResultSunInfo.UETransform.SetLocation(SunLocation);
+	ResultSunInfo.DirectionTowardEarth = (FVector::ZeroVector - SunLocation).GetSafeNormal();
+	return ResultSunInfo;
+}
+
+FPlanetaryBodyInfo ACelestialVaultDaySequenceActor::GetMoonInfo(double JulianDate)
+{
+	double UEDistance = CelestialVaultDistance * 100000 * MoonVaultPercentage / 100.0;
+	return GetPlanetaryBodyInfo(FPlanetaryBodyInputData::Moon, JulianDate, UEDistance, MoonScale);
+}
+
+bool ACelestialVaultDaySequenceActor::GetClosestStarInfo(FVector ObserverLocation, FVector LookupDirection, double ThresholdAngleDegrees, FStarInfo& FoundStarInfo, FTransform& StarTransform)
+{
+	StarTransform = FTransform::Identity;
+	
+	// We query only if we generated the StarInfo  
+	if (!KeepStarsInfo || !StarsComponent)
+	{
+		return false;
+	}
+	
+	double CosThresholdAngle = FMath::Cos(FMath::DegreesToRadians(ThresholdAngleDegrees));
+	LookupDirection.Normalize();
+	double MinCos = -1.0;
+	int32 ClosestInstanceIndex = -1;
+	
+	for (int StarInfoIndex = 0; StarInfoIndex < StarsInfo.Num(); StarInfoIndex++)
+	{
+		FTransform ISMInstanceTransform;
+		StarsComponent->GetInstanceTransform(StarsInfo[StarInfoIndex].ISMInstanceIndex, ISMInstanceTransform, true );
+
+		FVector DirectionToInstance = ISMInstanceTransform.GetLocation() - ObserverLocation;
+		DirectionToInstance.Normalize();
+
+		double CosDeltaAngle = FVector::DotProduct(LookupDirection, DirectionToInstance);
+		if (CosDeltaAngle > CosThresholdAngle && CosDeltaAngle > MinCos)
+		{
+			// Inside the cone angle, and closer to latest
+			MinCos = CosDeltaAngle;
+			ClosestInstanceIndex = StarInfoIndex;
+		}
+	}
+
+	if (ClosestInstanceIndex != -1)
+	{
+		StarsComponent->GetInstanceTransform(StarsInfo[ClosestInstanceIndex].ISMInstanceIndex, StarTransform, true ); 
+		FoundStarInfo = StarsInfo[ClosestInstanceIndex];
+		return true;
+	}
+
+	return false;
+}
+
+bool ACelestialVaultDaySequenceActor::GetClosestPlanetaryBody(FVector ObserverPosition, FVector LookupDirection, double ThresholdAngleDegrees, FPlanetaryBodyInfo& FoundPlanetaryBodyInfo, FTransform& BodyTransform)
+{
+	BodyTransform = FTransform::Identity;
+	
+	// We query only if we generated the StarInfo  
+	if (!KeepPlanetsInfos)
+	{
+		return false;
+	}
+	
+	double CosThresholdAngle = FMath::Cos(FMath::DegreesToRadians(ThresholdAngleDegrees));
+	LookupDirection.Normalize();
+	double MinCos = -1.0;
+	int32 ClosestInstanceIndex = -1;
+
+	for (FPlanetaryBodyInfo BodyInfo : PlanetsInfos)
+	{
+		if (PlanetsComponent)
+		{
+			// We need to use the ISM component and query the world transform because the celestial vault has rotated
+			FTransform ISMInstanceTransform;
+			PlanetsComponent->GetInstanceTransform(BodyInfo.ISMInstanceIndex, ISMInstanceTransform, true );
+
+			FVector DirectionToInstance = ISMInstanceTransform.GetLocation() - ObserverPosition;
+			DirectionToInstance.Normalize();
+
+			double CosDeltaAngle = FVector::DotProduct(LookupDirection, DirectionToInstance);
+			if (CosDeltaAngle > CosThresholdAngle && CosDeltaAngle > MinCos)
+			{
+				// Inside the cone angle, and closer to latest
+				MinCos = CosDeltaAngle;
+				ClosestInstanceIndex = BodyInfo.ISMInstanceIndex;
+			}
+		}
+	}
+
+	// Check for the moon
+	bool bIsMoonCloser = false;
+	if (MoonDiscComponent)
+	{
+		// We need to use the ISM component and query the world transform because the celestial vault has rotated
+		FVector MoonLocation = MoonDiscComponent->GetComponentLocation();
+		FVector DirectionToMoon = MoonLocation - ObserverPosition;
+		DirectionToMoon.Normalize();
+
+		double CosDeltaAngle = FVector::DotProduct(LookupDirection, DirectionToMoon);
+		if (CosDeltaAngle > CosThresholdAngle && CosDeltaAngle > MinCos)
+		{
+			// Inside the cone angle, and closer to latest
+			bIsMoonCloser = true;
+		}
+	}
+
+	if (bIsMoonCloser)
+	{
+		BodyTransform = MoonDiscComponent->GetComponentTransform();
+		FoundPlanetaryBodyInfo = MoonBodyInfo;
+		return true;
+	}
+	else if (ClosestInstanceIndex != -1)
+	{
+		PlanetsComponent->GetInstanceTransform(ClosestInstanceIndex, BodyTransform, true );
+		FoundPlanetaryBodyInfo = PlanetsInfos[ClosestInstanceIndex];
+		return true;
+	}
+
+	return false;
+}
+
+bool ACelestialVaultDaySequenceActor::GetPlanetaryBodyByOrbitType(EOrbitType OrbitType, FPlanetaryBodyInfo& FoundPlanetaryBodyInfo, FTransform& BodyTransform)
+{
+	BodyTransform = FTransform::Identity;
+	
+	// We query only if we generated the StarInfo  
+	if (!KeepPlanetsInfos)
+	{
+		return false;
+	}
+	
+	if (OrbitType == EOrbitType::Moon && MoonDiscComponent) 
+	{
+		BodyTransform = MoonDiscComponent->GetComponentTransform();
+		FoundPlanetaryBodyInfo = MoonBodyInfo;
+		return true;
+	}
+	else
+	{
+		for (FPlanetaryBodyInfo PlanetaryBodyInfo : PlanetsInfos)
+		{
+			if (PlanetaryBodyInfo.OrbitType == OrbitType)
+			{
+				if (PlanetsComponent)
+				{
+					PlanetsComponent->GetInstanceTransform(PlanetaryBodyInfo.ISMInstanceIndex, BodyTransform, true );
+				}
+				
+				FoundPlanetaryBodyInfo = PlanetaryBodyInfo;
+				return true;
+			}
+		}	
+	}
+	return false;
+}
+
+double ACelestialVaultDaySequenceActor::GetDayCelestialVaultAngle() const
+{
+	double TimeOffset = GMT_TimeZone;
+	if (bIsDaylightSavings)
+	{
+		TimeOffset += 1.0;
+	}
+	return GMST_AtTOD_0 - TimeOffset * 15.0; 
+}
+
+#if WITH_EDITOR
+
+void ACelestialVaultDaySequenceActor::PostEditChangeProperty(FPropertyChangedEvent& PropertyChangedEvent)
+{
+	//Get the name of the property that was changed  
+	FName PropertyName = (PropertyChangedEvent.Property != nullptr) ? PropertyChangedEvent.Property->GetFName() : NAME_None;
+
+	bool BuildAll = false;
+	bool BuildStars = false;
+	bool BuildPlanetaryBodies = false;
+	bool RebuildSequence = false;
+
+	// Geometry properties
+	if (PropertyName == GET_MEMBER_NAME_CHECKED(ACelestialVaultDaySequenceActor, CelestialVaultDistance) ||
+		PropertyName == GET_MEMBER_NAME_CHECKED(ACelestialVaultDaySequenceActor, StarsVaultPercentage) ||
+		PropertyName == GET_MEMBER_NAME_CHECKED(ACelestialVaultDaySequenceActor, PlanetsVaultPercentage) ||
+		PropertyName == GET_MEMBER_NAME_CHECKED(ACelestialVaultDaySequenceActor, MoonVaultPercentage))
+	{
+		if (DeepSkyComponent)
+		{
+			DeepSkyComponent->SetRelativeScale3D(FVector(CelestialVaultDistance*1000.0));	
+		}
+		
+		// TODO UpdateGeoReferencing();
+
+		BuildAll |= true;
+		RebuildSequence = true;
+	}
+
+	// Time properties
+	if (PropertyName == GET_MEMBER_NAME_CHECKED(ACelestialVaultDaySequenceActor, bUseCurrentDate) ||
+		PropertyName == GET_MEMBER_NAME_CHECKED(ACelestialVaultDaySequenceActor, Year) ||
+		PropertyName == GET_MEMBER_NAME_CHECKED(ACelestialVaultDaySequenceActor, Month) ||
+		PropertyName == GET_MEMBER_NAME_CHECKED(ACelestialVaultDaySequenceActor, Day) ||
+		PropertyName == GET_MEMBER_NAME_CHECKED(ACelestialVaultDaySequenceActor, GMT_TimeZone) ||
+		PropertyName == GET_MEMBER_NAME_CHECKED(ACelestialVaultDaySequenceActor, bIsDaylightSavings) ||
+		PropertyName == GET_MEMBER_NAME_CHECKED(ACelestialVaultDaySequenceActor, Latitude) ||
+		PropertyName == GET_MEMBER_NAME_CHECKED(ACelestialVaultDaySequenceActor, Longitude))
+	{
+		// DateTime has changed, we need to update the Reference CelestialVault Angle.
+		FDateTime DateTime = GetDate();
+		GMST_AtTOD_0 = UCelestialMaths::DateTimeToGreenwichMeanSiderealTime(DateTime);
+		double CelestialVaultAngle = GetDayCelestialVaultAngle();
+		CelestialVaultComponent->SetRelativeRotation( FRotator(0, CelestialVaultAngle, 0.0) );
+
+		BuildAll = true;
+		RebuildSequence = true;
+	}
+
+	// Properties impacting the Stars
+	if (PropertyName == GET_MEMBER_NAME_CHECKED(ACelestialVaultDaySequenceActor, CelestialStarCatalog) ||
+		PropertyName == GET_MEMBER_NAME_CHECKED(ACelestialVaultDaySequenceActor, FictionalStarCatalog) ||
+		PropertyName == GET_MEMBER_NAME_CHECKED(ACelestialVaultDaySequenceActor, MaxVisibleMagnitude) ||
+		PropertyName == GET_MEMBER_NAME_CHECKED(ACelestialVaultDaySequenceActor, KeepStarsInfo))
+	{
+		BuildStars |= true;
+		RebuildSequence = true;
+	}
+
+	// Properties impacting the Planetary Bodies
+	if (PropertyName == GET_MEMBER_NAME_CHECKED(ACelestialVaultDaySequenceActor, PlanetsCatalog) ||
+		PropertyName == GET_MEMBER_NAME_CHECKED(ACelestialVaultDaySequenceActor, PlanetsScale) ||
+		PropertyName == GET_MEMBER_NAME_CHECKED(ACelestialVaultDaySequenceActor, KeepPlanetsInfos) ||
+		PropertyName == GET_MEMBER_NAME_CHECKED(ACelestialVaultDaySequenceActor, MoonScale))
+	{
+		BuildPlanetaryBodies |= true;
+		RebuildSequence = true;
+	}
+
+	if (BuildAll || BuildStars)
+	{
+		InitStars();
+	}
+
+	if (BuildAll || BuildPlanetaryBodies)
+	{
+		InitPlanetaryBodies();
+	}
+
+	if (RebuildSequence)
+	{
+		RootSequence = nullptr;
+		UpdateRootSequence();
+	}
+
+	Super::PostEditChangeProperty(PropertyChangedEvent);
+}
+#endif
+
+//////// Private  Functions
+
+void ACelestialVaultDaySequenceActor::InitStars() 
+{
+	TRACE_CPUPROFILER_EVENT_SCOPE(ACelestialVaultDaySequenceActor::InitStars);
+	
+	if (StarsComponent)
+	{
+		StarsComponent->ClearInstances();
+		StarsComponent->SetNumCustomDataFloats(4);
+		StarsInfo.Empty();
+		
+		// Check for the Catalog
+		if (!CelestialStarCatalog && !FictionalStarCatalog)
+		{
+			UE_LOG(LogCelestialVault, Warning, TEXT("Please define at least a Celestial or Fictional StarCatalog"));
+			return;
+		}
+
+		if (CelestialStarCatalog)
+		{
+			// Check for the Catalog Data types
+			const UScriptStruct* CelestialStarCatalogRowStruct = CelestialStarCatalog->GetRowStruct();
+			if (CelestialStarCatalogRowStruct != FCelestialStarInputData::StaticStruct())
+			{
+				UE_LOG(LogCelestialVault, Warning, TEXT("Invalid DataTable row structure for CelestialStarCatalog! It should be of type %s"), *FCelestialStarInputData::StaticStruct()->GetName() );
+			}
+			else
+			{
+				// Generate Celestial Stars
+				CelestialStarCatalog->ForeachRow<FCelestialStarInputData>("AEarthDaySequenceActor::InitStars", [&](const FName& Key, const FCelestialStarInputData& CelestialStarInputData)
+					{
+						if (CelestialStarInputData.Magnitude <= MaxVisibleMagnitude)
+						{
+							// Location - Convert to UE Left-handed Frame (Invert the Y coordinate)
+							FVector StarLocation = UCelestialMaths::RADECToXYZ_RH(CelestialStarInputData.RA * 15, CelestialStarInputData.DEC, CelestialVaultDistance * 1000.0 * StarsVaultPercentage / 100.0) *  FVector(100, -100, 100);
+
+							// Color
+							FLinearColor StarColor = UCelestialMaths::BVtoLinearColor(CelestialStarInputData.ColorIndex);
+
+							// Create ISM instance
+							int32 NewIndex = StarsComponent->AddInstance(FTransform(StarLocation), false);
+							TArray<float> NewCustomData;
+							NewCustomData.Add(StarColor.R);
+							NewCustomData.Add(StarColor.G);
+							NewCustomData.Add(StarColor.B);
+							NewCustomData.Add(CelestialStarInputData.Magnitude);
+							StarsComponent->SetCustomData(NewIndex, NewCustomData);
+
+							// Keep trace of the Star Information for further runtime queries					
+							if (KeepStarsInfo)
+							{
+								FStarInfo StarInfo;
+								StarInfo.RA = CelestialStarInputData.RA;
+								StarInfo.DEC = CelestialStarInputData.DEC;
+								StarInfo.DistanceInPC = CelestialStarInputData.DistanceInPC;
+								StarInfo.Name = CelestialStarInputData.Name;
+								StarInfo.Magnitude = CelestialStarInputData.Magnitude;
+								StarInfo.Color = StarColor;
+								StarInfo.HipparcosID = CelestialStarInputData.HipparcosID;
+								StarInfo.HenryDraperID = CelestialStarInputData.HenryDraperID;
+								StarInfo.YaleBrightStarID = CelestialStarInputData.YaleBrightStarID;
+								StarInfo.ColorIndex = CelestialStarInputData.ColorIndex;
+								StarInfo.ISMInstanceIndex = NewIndex;
+								// Maybe add other computed values here... 
+								StarsInfo.Add(StarInfo);
+							}
+						}
+					}
+				);
+			}
+		}
+
+		if (FictionalStarCatalog)
+		{
+			const UScriptStruct* FictionalStarCatalogRowStruct = FictionalStarCatalog->GetRowStruct();
+			if (FictionalStarCatalogRowStruct != FStarInputData::StaticStruct())
+			{
+				UE_LOG(LogCelestialVault, Warning, TEXT("Invalid DataTable row structure for FictionalStarCatalog! It should be of type %s"), *FStarInputData::StaticStruct()->GetName() );
+			}
+			else
+			{
+				// Generate Fictional Stars
+				FictionalStarCatalog->ForeachRow<FStarInputData>("AEarthDaySequenceActor::InitStars / Fictional", [&](const FName& Key, const FStarInputData& StarInputData)
+					{
+						if (StarInputData.Magnitude <= MaxVisibleMagnitude)
+						{
+							// Location - Convert to UE Left-handed Frame (Invert the Y coordinate)
+							FVector StarLocation = UCelestialMaths::RADECToXYZ_RH(StarInputData.RA * 15, StarInputData.DEC, CelestialVaultDistance * 1000.0 * StarsVaultPercentage / 100.0) * FVector(100, -100, 100);
+
+							// Color
+							FLinearColor StarColor = StarInputData.Color;
+
+							// Create ISM instance
+							int32 NewIndex = StarsComponent->AddInstance(FTransform(StarLocation), false);
+							TArray<float> NewCustomData;
+							NewCustomData.Add(StarColor.R);
+							NewCustomData.Add(StarColor.G);
+							NewCustomData.Add(StarColor.B);
+							NewCustomData.Add(StarInputData.Magnitude);
+							StarsComponent->SetCustomData(NewIndex, NewCustomData);
+
+							// Keep trace of the Star Information for further runtime queries					
+							if (KeepStarsInfo)
+							{
+								FStarInfo StarInfo;
+								StarInfo.RA = StarInputData.RA;
+								StarInfo.DEC = StarInputData.DEC;
+								StarInfo.DistanceInPC = StarInputData.DistanceInPC;
+								StarInfo.Name = StarInputData.Name;
+								StarInfo.Magnitude = StarInputData.Magnitude;
+								StarInfo.Color = StarColor;
+								StarInfo.ISMInstanceIndex = NewIndex;
+								StarsInfo.Add(StarInfo);
+							}
+						}
+					}
+				);
+			}
+		}
+		
+		UE_LOG(LogCelestialVault, Verbose, TEXT("%d Stars added "), StarsComponent->GetInstanceCount());
+		StarsComponent->MarkRenderInstancesDirty();
+	}		
+}
+
+void ACelestialVaultDaySequenceActor::InitPlanetaryBodies()
+{
+	TRACE_CPUPROFILER_EVENT_SCOPE(ACelestialVaultDaySequenceActor::InitPlanetaryBodies);
+	
+	// Get the proper Julian day, at Midnight... The Daysequence will rotate the sky vault later 
+	FDateTime LocalTimeMidnight = GetDate();
+	FDateTime UTCTime = UCelestialMaths::LocalTimeToUTCTime(LocalTimeMidnight, GMT_TimeZone, bIsDaylightSavings);
+	double JulianDay = UCelestialMaths::UTCDateTimeToJulianDate(UTCTime);
+
+	// Init the sun
+	SunInfo = GetSunInfo(JulianDay);
+	
+	// Init the Moon
+	if (MoonDiscComponent)
+	{
+		// We don't want to use the GetMoonTransform function here because we also want to keep the Moon information for the day.
+		// It has a pitfall because the Moon information will not be adjusted to the current TOD, but that doesn't prevent anyone for doing a manuel query at some point
+		MoonBodyInfo = GetMoonInfo(JulianDay);
+		MoonDiscComponent->SetRelativeTransform(MoonBodyInfo.UETransform);
+	}
+	
+	// Init the Planets from the catalog
+	if (PlanetsComponent)
+	{
+		PlanetsComponent->ClearInstances();
+		PlanetsComponent->SetNumCustomDataFloats(2);
+		PlanetsInfos.Empty();
+
+		// Safety check on Catalog Data
+		if (!PlanetsCatalog)
+		{
+			UE_LOG(LogCelestialVault, Warning, TEXT("PlanetaryBodiesCatalog is null!"));
+			return;
+		}
+
+		const UScriptStruct* RowStruct = PlanetsCatalog->GetRowStruct();
+		if (RowStruct != FPlanetaryBodyInputData::StaticStruct() )
+		{
+			UE_LOG(LogCelestialVault, Warning, TEXT("Invalid DataTable row structure for the Planetary Bodies Catalog! It should be of type %s"), *FPlanetaryBodyInputData::StaticStruct()->GetName() );
+			return;
+		}
+
+		// Init from the DataTable
+		PlanetsCatalog->ForeachRow<FPlanetaryBodyInputData>("AEarthDaySequenceActor::InitPlanetaryBodies", [&](const FName& Key, const FPlanetaryBodyInputData& InputPlanetaryBody)
+			{
+				double UEBodyDistance = CelestialVaultDistance * 100000 * PlanetsVaultPercentage / 100.0;
+				FPlanetaryBodyInfo BodyInfo = GetPlanetaryBodyInfo(InputPlanetaryBody, JulianDay, UEBodyDistance, PlanetsScale);
+
+				// Add the new ISM Instance
+				int32 NewIndex = PlanetsComponent->AddInstance(BodyInfo.UETransform);
+				BodyInfo.ISMInstanceIndex = NewIndex;
+			
+				// Add the Custom data (column index to sample the planets atlas texture)
+				TArray<float> NewCustomData;
+				NewCustomData.Add(InputPlanetaryBody.TextureColumnIndex);
+				NewCustomData.Add(BodyInfo.ApparentMagnitude);
+				PlanetsComponent->SetCustomData(NewIndex, NewCustomData);
+
+				// Keep trace of the PlanetaryBody for further queries (the Datatable is readonly, so store in another object)
+				if (KeepPlanetsInfos)
+				{
+					PlanetsInfos.Add(BodyInfo);	
+				}
+			}
+		);
+		PlanetsComponent->MarkRenderInstancesDirty();
+	}
+}
+
+FPlanetaryBodyInfo ACelestialVaultDaySequenceActor::GetPlanetaryBodyInfo(const FPlanetaryBodyInputData& InputPlanetaryBody, double JulianDay, double UEDistance, double BodyScale) const
+{
+	FPlanetaryBodyInfo BodyInfo;
+	BodyInfo.OrbitType = InputPlanetaryBody.OrbitType;
+	BodyInfo.Name = InputPlanetaryBody.Name;
+	BodyInfo.Radius = InputPlanetaryBody.Radius;
+
+	// Compute location	
+	double RAHours, DECDegrees, DistanceToEarthAU, DistanceToSunAU, DistanceEarthToSunAU;
+	UCelestialMaths::GetBodyCelestialCoordinatesAU(JulianDay, InputPlanetaryBody, Latitude, Longitude, RAHours, DECDegrees, DistanceToEarthAU, DistanceToSunAU, DistanceEarthToSunAU);
+	BodyInfo.RA = RAHours;
+	BodyInfo.DEC = DECDegrees;
+	BodyInfo.DistanceInAU = DistanceToEarthAU;
+
+	// Compute Magnitude
+	double Phase; // TODO - Check the PhaseAngle Computations. 
+	BodyInfo.ApparentMagnitude = UCelestialMaths::GetPlanetaryBodyMagnitude(InputPlanetaryBody, DistanceToEarthAU, DistanceToSunAU, DistanceEarthToSunAU, Phase);
+	BodyInfo.Phase = Phase;
+	
+	// Compute the True and the Scaled apparent Diameters
+	BodyInfo.ApparentDiameterDegrees = FMath::RadiansToDegrees(FMath::Atan2(InputPlanetaryBody.Radius * 1000.0, UCelestialMaths::AstronomicalUnitsToMeters(DistanceToEarthAU)) * 2);
+	BodyInfo.ScaledApparentDiameterDegrees = FMath::RadiansToDegrees(FMath::Atan2(InputPlanetaryBody.Radius * 1000.0 * BodyScale, UCelestialMaths::AstronomicalUnitsToMeters(DistanceToEarthAU)) * 2);
+
+	// Location
+	FVector BodyLocation = UCelestialMaths::RADECToXYZ_RH(BodyInfo.RA * 15, BodyInfo.DEC, UEDistance);
+	BodyLocation *= FVector(1, -1, 1); // Convert to UE Frame by inverting Y
+
+	// Rotation - Orient the Body mesh towards the earth -
+	// Use Look at the UE Origin, this is an acceptable approximation if the vault is big enough - We could use the Earth center if needed
+	FVector BodyToEarth = (FVector::ZeroVector - BodyLocation).GetSafeNormal();
+	FQuat BodyRotation = FQuat(BodyToEarth.Rotation());
+
+	// Scale
+	// The mesh plane is 100 UE Units (1m), and located at PlanetMeshDistanceKm. Use Thales theorem to compute its effective scale at this distance to ensure the right apparent diameter
+	double UERadius = (UEDistance * InputPlanetaryBody.Radius * 1000.0) / UCelestialMaths::AstronomicalUnitsToMeters(DistanceToEarthAU);
+	double Scale = UERadius * BodyScale / 50.0; // The plane half-length is 50 UE Units
+	FVector BodyScale3D = FVector(Scale, Scale, Scale);
+
+	BodyInfo.UETransform = FTransform(BodyRotation, BodyLocation, BodyScale3D);
+
+	// Direction
+	BodyInfo.DirectionTowardEarth = BodyToEarth;
+	
+	return BodyInfo;
+}
+
+void ACelestialVaultDaySequenceActor::RebuildAll()
+{
+	if (const UWorld* World = GetWorld())
+	{
+		// DateTime has changed, we need to update the Reference CelestialVault Angle.
+		FDateTime DateTime = GetDate();
+		GMST_AtTOD_0 = UCelestialMaths::DateTimeToGreenwichMeanSiderealTime(DateTime);
+
+		// Get the transformation in proper Celestial World units (meters, right handed)
+		FTransform ECEFFrameToWorldFrame = UCelestialMaths::GetPlanetCenterTransform(Latitude, Longitude, 0);
+
+		// UE Frame are expressed in Left-handed coordinates, and units are in meters - Convert to UE Transform 
+		FMatrix WorldFrameToUEFrame = FMatrix( 
+			FVector(1.0, 0.0, 0.0),		// Easting (X) is UE World X
+			FVector(0.0, -1.0, 0.0),	// Northing (Y) is UE World -Y because of left-handed convention
+			FVector(0.0, 0.0, 1.0),		// Up (Z) is UE World Z 
+			FVector(0.0, 0.0, 0.0));	// No Origin offset
+		FMatrix UEFrameToWorldFrame = WorldFrameToUEFrame.Inverse();
+
+		// Update the rotation part
+		FMatrix TransformMatrix = UEFrameToWorldFrame * ECEFFrameToWorldFrame.ToMatrixNoScale() * WorldFrameToUEFrame;
+		// Get Origin, and convert UE units to meters
+		FVector UEOrigin = TransformMatrix.GetOrigin() * FVector(100.0, 100.0, 100.0);
+		TransformMatrix.SetOrigin(UEOrigin);
+
+		// Apply the transform
+		PlanetCenterTransform = FTransform(TransformMatrix);
+		PlanetCenterComponent->SetWorldTransform(PlanetCenterTransform);
+
+		// Rebuild Sky and Sequence
+		InitStars();
+		InitPlanetaryBodies();
+
+		double CelestialVaultAngle = GetDayCelestialVaultAngle();
+		CelestialVaultComponent->SetRelativeRotation( FRotator(0, CelestialVaultAngle, 0.0) );
+	}
+}
diff --git a/Engine/Plugins/Experimental/CelestialVault/Source/CelestialVault/Private/CelestialVaultSequence.cpp b/Engine/Plugins/Experimental/CelestialVault/Source/CelestialVault/Private/CelestialVaultSequence.cpp
new file mode 100644
--- /dev/null
+++ b/Engine/Plugins/Experimental/CelestialVault/Source/CelestialVault/Private/CelestialVaultSequence.cpp
@@ -0,0 +1,93 @@
+﻿// Copyright Epic Games, Inc. All Rights Reserved.
+
+
+#include "CelestialVaultSequence.h"
+
+#include "CelestialMaths.h"
+#include "CelestialVault.h"
+#include "CelestialVaultDaySequenceActor.h"
+
+#include "Components/StaticMeshComponent.h"
+#include "Components/DirectionalLightComponent.h"
+
+
+void FCelestialVaultSequence::BuildSequence(UProceduralDaySequenceBuilder* InBuilder)
+{
+	TRACE_CPUPROFILER_EVENT_SCOPE(FCelestialVaultSequence::BuildSequence);
+	
+	ACelestialVaultDaySequenceActor* EarthDaySequenceActor = Cast<ACelestialVaultDaySequenceActor>(WeakTargetActor.Get());
+	if (!EarthDaySequenceActor)
+	{
+		UE_LOG(LogCelestialVault, Warning, TEXT("This ProceduralDaySequence applies only on AEarthDaySequenceActors"));
+		return;
+	}
+
+	// Use the reference time for Midnight, and build the keys from here
+	FDateTime LocalTimeMidnight = EarthDaySequenceActor->GetDate();
+
+	
+	// Create the Celestial Vault Animation
+	if (USceneComponent* CelestialVaultComponent = EarthDaySequenceActor->CelestialVaultComponent)
+	{
+		InBuilder->SetActiveBoundObject(CelestialVaultComponent);
+		
+		const double NormalizedTimeIncrement = 1.0 / FMath::Max(KeyCount - 1, static_cast<unsigned>(1));
+		for (unsigned int Key = 0; Key < KeyCount; ++Key)
+		{
+			double CelestialVaultAngleAtMidnight = EarthDaySequenceActor->GetDayCelestialVaultAngle(); 
+
+			const double KeyTime = Key * NormalizedTimeIncrement;
+			InBuilder->AddRotationKey(KeyTime, FRotator(0, CelestialVaultAngleAtMidnight + KeyTime * 360, 0.0));
+		}
+	}
+
+	// Create the Moon Disc Animation
+	UStaticMeshComponent* MoonDiscComponent = EarthDaySequenceActor->MoonDiscComponent;
+	UDirectionalLightComponent* MoonLightComponent = EarthDaySequenceActor->MoonLightComponent;
+	
+	if (MoonDiscComponent && MoonLightComponent)
+	{
+		const double NormalizedTimeIncrement = 1.0 / FMath::Max(KeyCount - 1, static_cast<unsigned>(1));
+		for (unsigned int Key = 0; Key < KeyCount; ++Key)
+		{
+			const double KeyTime = Key * NormalizedTimeIncrement;
+
+			FDateTime KeyDateTime = LocalTimeMidnight + FTimespan::FromSeconds(KeyTime * 24*60*60);
+			FDateTime UTCTime = UCelestialMaths::LocalTimeToUTCTime(KeyDateTime, EarthDaySequenceActor->GMT_TimeZone, EarthDaySequenceActor->bIsDaylightSavings);
+			double JulianDay = UCelestialMaths::UTCDateTimeToJulianDate(UTCTime);
+
+			FPlanetaryBodyInfo MoonInfo = EarthDaySequenceActor->GetMoonInfo(JulianDay);
+
+			// Animate Disc
+			InBuilder->SetActiveBoundObject(MoonDiscComponent);
+			InBuilder->AddTransformKey(KeyTime, MoonInfo.UETransform);
+
+			// Animate Light
+			InBuilder->SetActiveBoundObject(MoonLightComponent);
+			InBuilder->AddRotationKey(KeyTime, FRotationMatrix::MakeFromX(MoonInfo.DirectionTowardEarth).Rotator());
+			
+		}
+	}
+
+	// Create the SunLight Animation
+	if (UDirectionalLightComponent* SunLightComponent = EarthDaySequenceActor->SunLightComponent)
+	{
+		InBuilder->SetActiveBoundObject(SunLightComponent);
+		
+		const double NormalizedTimeIncrement = 1.0 / FMath::Max(KeyCount - 1, static_cast<unsigned>(1));
+		for (unsigned int Key = 0; Key < KeyCount; ++Key)
+		{
+			const double KeyTime = Key * NormalizedTimeIncrement;
+
+			FDateTime KeyDateTime = LocalTimeMidnight + FTimespan::FromSeconds(KeyTime * 24*60*60);
+			FDateTime UTCTime = UCelestialMaths::LocalTimeToUTCTime(KeyDateTime, EarthDaySequenceActor->GMT_TimeZone, EarthDaySequenceActor->bIsDaylightSavings);
+			double JulianDay = UCelestialMaths::UTCDateTimeToJulianDate(UTCTime);
+
+			FSunInfo SunInfo = EarthDaySequenceActor->GetSunInfo(JulianDay);
+			// Animate Light
+			InBuilder->AddRotationKey(KeyTime, FRotationMatrix::MakeFromX(SunInfo.DirectionTowardEarth).Rotator());
+		}
+	}
+	
+	FProceduralDaySequence::BuildSequence(InBuilder);
+}
