Added the first version of the Celestial Vault DaySequence plugin

diff --git a/Engine/Plugins/Experimental/CelestialVault/Source/CelestialVault/Private/vsop87a_full_EarthMoon.cpp b/Engine/Plugins/Experimental/CelestialVault/Source/CelestialVault/Private/vsop87a_full_EarthMoon.cpp
new file mode 100644
--- /dev/null
+++ b/Engine/Plugins/Experimental/CelestialVault/Source/CelestialVault/Private/vsop87a_full_EarthMoon.cpp
@@ -0,0 +1,7016 @@
+// Copyright Epic Games, Inc. All Rights Reserved.
+
+
+//VSOP87-Multilang http://www.celestialprogramming.com/vsop87-multilang/index.html
+//Greg Miller ([REDACTED_EMAIL]) 2019.  Released as Public Domain
+
+#include "vsop87a_full_EarthMoon.h"
+#include <math.h>
+
+
+void vsop87a_full::getEarth(double t,double temp[]){
+   temp[0]=vsop87a_full::earth_x(t);
+   temp[1]=vsop87a_full::earth_y(t);
+   temp[2]=vsop87a_full::earth_z(t);
+}
+
+void vsop87a_full::getEmb(double t,double temp[]){
+   temp[0]=vsop87a_full::emb_x(t);
+   temp[1]=vsop87a_full::emb_y(t);
+   temp[2]=vsop87a_full::emb_z(t);
+}
+
+void vsop87a_full::getMoon(double earth[], double emb[],double temp[]){
+   temp[0]=(emb[0]-earth[0])*(1 + 1 / 0.01230073677);
+   temp[1]=(emb[1]-earth[1])*(1 + 1 / 0.01230073677);
+   temp[2]=(emb[2]-earth[2])*(1 + 1 / 0.01230073677);
+   temp[0]=temp[0]+earth[0];
+   temp[1]=temp[1]+earth[1];
+   temp[2]=temp[2]+earth[2];
+}
+
+double vsop87a_full::earth_x(double t){
+   double earth_x_0=0.0;
+   earth_x_0+=     0.99982928844 * cos( 1.75348568475 +     6283.07584999140*t);
+   earth_x_0+=     0.00835257300 * cos( 1.71034539450 +    12566.15169998280*t);
+   earth_x_0+=     0.00561144206 * cos( 0.00000000000 +        0.00000000000*t);
+   earth_x_0+=     0.00010466628 * cos( 1.66722645223 +    18849.22754997420*t);
+   earth_x_0+=     0.00003110838 * cos( 0.66875185215 +    83996.84731811189*t);
+   earth_x_0+=     0.00002552498 * cos( 0.58310207301 +      529.69096509460*t);
+   earth_x_0+=     0.00002137256 * cos( 1.09235189672 +     1577.34354244780*t);
+   earth_x_0+=     0.00001709103 * cos( 0.49540223397 +     6279.55273164240*t);
+   earth_x_0+=     0.00001707882 * cos( 6.15315547484 +     6286.59896834040*t);
+   earth_x_0+=     0.00001445242 * cos( 3.47272783760 +     2352.86615377180*t);
+   earth_x_0+=     0.00001091006 * cos( 3.68984782465 +     5223.69391980220*t);
+   earth_x_0+=     0.00000934429 * cos( 6.07389922585 +    12036.46073488820*t);
+   earth_x_0+=     0.00000899144 * cos( 3.17571950523 +    10213.28554621100*t);
+   earth_x_0+=     0.00000566514 * cos( 2.15262034016 +     1059.38193018920*t);
+   earth_x_0+=     0.00000684416 * cos( 1.30699021227 +     5753.38488489680*t);
+   earth_x_0+=     0.00000734455 * cos( 4.35500196530 +      398.14900340820*t);
+   earth_x_0+=     0.00000681437 * cos( 2.21821534685 +     4705.73230754360*t);
+   earth_x_0+=     0.00000611238 * cos( 5.38479234323 +     6812.76681508600*t);
+   earth_x_0+=     0.00000451836 * cos( 6.08768280868 +     5884.92684658320*t);
+   earth_x_0+=     0.00000451953 * cos( 1.27933728354 +     6256.77753019160*t);
+   earth_x_0+=     0.00000449517 * cos( 5.36923831714 +     6309.37416979120*t);
+   earth_x_0+=     0.00000406248 * cos( 0.54361367084 +     6681.22485339960*t);
+   earth_x_0+=     0.00000540957 * cos( 0.78677364655 +      775.52261132400*t);
+   earth_x_0+=     0.00000547004 * cos( 1.46146650376 +    14143.49524243060*t);
+   earth_x_0+=     0.00000520484 * cos( 4.43295799975 +     7860.41939243920*t);
+   earth_x_0+=     0.00000214960 * cos( 4.50213844573 +    11506.76976979360*t);
+   earth_x_0+=     0.00000227892 * cos( 1.23941482802 +     7058.59846131540*t);
+   earth_x_0+=     0.00000225878 * cos( 3.27244306207 +     4694.00295470760*t);
+   earth_x_0+=     0.00000255820 * cos( 2.26556277246 +    12168.00269657460*t);
+   earth_x_0+=     0.00000256182 * cos( 1.45474116190 +      709.93304855830*t);
+   earth_x_0+=     0.00000178120 * cos( 2.96205424204 +      796.29800681640*t);
+   earth_x_0+=     0.00000161205 * cos( 1.47337718956 +     5486.77784317500*t);
+   earth_x_0+=     0.00000178325 * cos( 6.24374704602 +     6283.14316029419*t);
+   earth_x_0+=     0.00000178325 * cos( 0.40466470869 +     6283.00853968860*t);
+   earth_x_0+=     0.00000155487 * cos( 1.62409309523 +    25132.30339996560*t);
+   earth_x_0+=     0.00000209024 * cos( 5.85207528073 +    11790.62908865880*t);
+   earth_x_0+=     0.00000199971 * cos( 4.07209938245 +    17789.84561978500*t);
+   earth_x_0+=     0.00000128933 * cos( 5.21693314150 +     7079.37385680780*t);
+   earth_x_0+=     0.00000128099 * cos( 4.80182882228 +     3738.76143010800*t);
+   earth_x_0+=     0.00000151691 * cos( 0.86921639327 +      213.29909543800*t);
+   earth_x_0+=     0.00000094483 * cos( 4.60187620729 +    10977.07880469900*t);
+   earth_x_0+=     0.00000077801 * cos( 1.84444797200 +     8827.39026987480*t);
+   earth_x_0+=     0.00000077277 * cos( 3.58396027325 +     5507.55323866740*t);
+   earth_x_0+=     0.00000073493 * cos( 2.69522392584 +     1589.07289528380*t);
+   earth_x_0+=     0.00000065359 * cos( 3.65156347970 +    11769.85369316640*t);
+   earth_x_0+=     0.00000063240 * cos( 2.24124130944 +     6262.30045449900*t);
+   earth_x_0+=     0.00000062961 * cos( 4.40729381941 +     6303.85124548380*t);
+   earth_x_0+=     0.00000085870 * cos( 3.02430773140 +   167283.76158766549*t);
+   earth_x_0+=     0.00000082957 * cos( 6.19254908840 +     3340.61242669980*t);
+   earth_x_0+=     0.00000079201 * cos( 2.49366716519 +     6069.77675455340*t);
+   earth_x_0+=     0.00000060048 * cos( 3.39951805979 +     4136.91043351620*t);
+   earth_x_0+=     0.00000062119 * cos( 2.01472298915 +        7.11354700080*t);
+   earth_x_0+=     0.00000075601 * cos( 4.15956023781 +     6496.37494542940*t);
+   earth_x_0+=     0.00000067771 * cos( 4.03416283669 +     9437.76293488700*t);
+   earth_x_0+=     0.00000053715 * cos( 1.56222706033 +     1194.44701022460*t);
+   earth_x_0+=     0.00000057574 * cos( 2.67881766950 +     6282.09552892320*t);
+   earth_x_0+=     0.00000057574 * cos( 3.96971590315 +     6284.05617105960*t);
+   earth_x_0+=     0.00000051129 * cos( 1.31255522350 +     6290.18939699220*t);
+   earth_x_0+=     0.00000050529 * cos( 5.34827650648 +     6275.96230299060*t);
+   earth_x_0+=     0.00000049613 * cos( 2.64370071484 +     6127.65545055720*t);
+   earth_x_0+=     0.00000048085 * cos( 4.00886197380 +     6438.49624942560*t);
+   earth_x_0+=     0.00000057340 * cos( 3.79459979103 +     3154.68708489560*t);
+   earth_x_0+=     0.00000047378 * cos( 6.10705175562 +     3128.38876509580*t);
+   earth_x_0+=     0.00000048104 * cos( 4.77180566934 +      801.82093112380*t);
+   earth_x_0+=     0.00000041149 * cos( 3.32726302190 +     8429.24126646660*t);
+   earth_x_0+=     0.00000052309 * cos( 5.30571599873 +    13367.97263110660*t);
+   earth_x_0+=     0.00000051321 * cos( 5.78450026632 +    12352.85260454480*t);
+   earth_x_0+=     0.00000050648 * cos( 2.05206551992 +    11856.21865142450*t);
+   earth_x_0+=     0.00000043381 * cos( 3.64501749522 +    17260.15465469040*t);
+   earth_x_0+=     0.00000039554 * cos( 4.93022392385 +     5481.25491886760*t);
+   earth_x_0+=     0.00000048969 * cos( 0.45418080759 +     9225.53927328300*t);
+   earth_x_0+=     0.00000047568 * cos( 3.16123585999 +     5856.47765911540*t);
+   earth_x_0+=     0.00000047271 * cos( 0.62147119447 +     2544.31441988340*t);
+   earth_x_0+=     0.00000037987 * cos( 3.04100734504 +      426.59819087600*t);
+   earth_x_0+=     0.00000032481 * cos( 0.80957263520 +     3930.20969621960*t);
+   earth_x_0+=     0.00000032559 * cos( 1.97427864771 +     2146.16541647520*t);
+   earth_x_0+=     0.00000032521 * cos( 2.84566871337 +     4164.31198961300*t);
+   earth_x_0+=     0.00000032546 * cos( 3.01796808307 +     5088.62883976680*t);
+   earth_x_0+=     0.00000028515 * cos( 6.10933021271 +    12569.67481833180*t);
+   earth_x_0+=     0.00000036898 * cos( 2.91097134183 +    18073.70493865020*t);
+   earth_x_0+=     0.00000024936 * cos( 2.17935828787 +       26.29831979980*t);
+   earth_x_0+=     0.00000030322 * cos( 1.99340330773 +     4535.05943692440*t);
+   earth_x_0+=     0.00000028166 * cos( 0.48127325250 +    12562.62858163380*t);
+   earth_x_0+=     0.00000021880 * cos( 5.12635088590 +    11371.70468975820*t);
+   earth_x_0+=     0.00000021642 * cos( 4.63524156478 +    12559.03815298200*t);
+   earth_x_0+=     0.00000020480 * cos( 1.46713660207 +     7084.89678111520*t);
+   earth_x_0+=     0.00000026112 * cos( 0.30447108865 +    71430.69561812909*t);
+   earth_x_0+=     0.00000022867 * cos( 4.76437399404 +     8031.09226305840*t);
+   earth_x_0+=     0.00000018593 * cos( 3.38415795860 +     1748.01641306700*t);
+   earth_x_0+=     0.00000023250 * cos( 6.19030311409 +    18319.53658487960*t);
+   earth_x_0+=     0.00000017104 * cos( 0.58684867006 +     4933.20844033260*t);
+   earth_x_0+=     0.00000021687 * cos( 4.30348901447 +    10447.38783960440*t);
+   earth_x_0+=     0.00000021065 * cos( 3.80028088787 +     7477.52286021600*t);
+   earth_x_0+=     0.00000014285 * cos( 1.29528607840 +     2942.46342329160*t);
+   earth_x_0+=     0.00000013898 * cos( 4.59171958613 +     8635.94200376320*t);
+   earth_x_0+=     0.00000019220 * cos( 0.49428010004 +    17298.18232732620*t);
+   earth_x_0+=     0.00000018884 * cos( 2.42694414831 +   156137.47598479928*t);
+   earth_x_0+=     0.00000017889 * cos( 0.15750451509 +     1592.59601363280*t);
+   earth_x_0+=     0.00000013605 * cos( 5.19721055627 +    13095.84266507740*t);
+   earth_x_0+=     0.00000015049 * cos( 3.12084096050 +    16496.36139620240*t);
+   earth_x_0+=     0.00000013852 * cos( 6.16243027421 +     7632.94325965020*t);
+   earth_x_0+=     0.00000014421 * cos( 1.43791851415 +    20426.57109242200*t);
+   earth_x_0+=     0.00000013247 * cos( 3.60768174652 +    12139.55350910680*t);
+   earth_x_0+=     0.00000011583 * cos( 5.46076026245 +     5331.35744374080*t);
+   earth_x_0+=     0.00000014329 * cos( 6.02791958133 +     7342.45778018060*t);
+   earth_x_0+=     0.00000012375 * cos( 3.10941923273 +     6279.48542133960*t);
+   earth_x_0+=     0.00000012366 * cos( 3.53911446434 +     6286.66627864320*t);
+   earth_x_0+=     0.00000009951 * cos( 1.25735916596 +     7234.79425624200*t);
+   earth_x_0+=     0.00000012425 * cos( 1.06396563046 +    15110.46611986620*t);
+   earth_x_0+=     0.00000011004 * cos( 2.19276723240 +    10988.80815753500*t);
+   earth_x_0+=     0.00000011583 * cos( 4.05256206393 +     5729.50644714900*t);
+   earth_x_0+=     0.00000009046 * cos( 5.43016645841 +     9623.68827669120*t);
+   earth_x_0+=     0.00000008869 * cos( 5.04690941813 +     6148.01076995600*t);
+   earth_x_0+=     0.00000010507 * cos( 2.62758025060 +     6836.64525283380*t);
+   earth_x_0+=     0.00000011129 * cos( 4.17728063171 +    15720.83878487840*t);
+   earth_x_0+=     0.00000008614 * cos( 1.60161176919 +     6418.14093002680*t);
+   earth_x_0+=     0.00000008532 * cos( 2.46198384967 +    14712.31711645800*t);
+   earth_x_0+=     0.00000007910 * cos( 3.15429574843 +     2118.76386037840*t);
+   earth_x_0+=     0.00000010546 * cos( 4.79503596728 +     1349.86740965880*t);
+   earth_x_0+=     0.00000010480 * cos( 2.95288908781 +     5999.21653112620*t);
+   earth_x_0+=     0.00000007456 * cos( 5.42522643084 +     6040.34724601740*t);
+   earth_x_0+=     0.00000009820 * cos( 3.69556194245 +     6566.93516885660*t);
+   earth_x_0+=     0.00000009171 * cos( 3.92216439653 +     5643.17856367740*t);
+   earth_x_0+=     0.00000007018 * cos( 1.23420199826 +     6525.80445396540*t);
+   earth_x_0+=     0.00000009797 * cos( 1.91951822764 +    21228.39202354580*t);
+   earth_x_0+=     0.00000009060 * cos( 4.61643194755 +     4690.47983635860*t);
+   earth_x_0+=     0.00000007545 * cos( 1.23604549224 +    12539.85338018300*t);
+   earth_x_0+=     0.00000008227 * cos( 0.32929613711 +    10973.55568635000*t);
+   earth_x_0+=     0.00000007202 * cos( 5.88391543785 +     6245.04817735560*t);
+   earth_x_0+=     0.00000007144 * cos( 0.76464650349 +     6321.10352262720*t);
+   earth_x_0+=     0.00000006625 * cos( 0.70112927898 +     5327.47610838280*t);
+   earth_x_0+=     0.00000007412 * cos( 5.03197290869 +    23543.23050468179*t);
+   earth_x_0+=     0.00000006160 * cos( 0.94570829909 +    12964.30070339100*t);
+   earth_x_0+=     0.00000006378 * cos( 6.20903842618 +      951.71840625060*t);
+   earth_x_0+=     0.00000006359 * cos( 5.03541261021 +     1990.74501704100*t);
+   earth_x_0+=     0.00000007254 * cos( 3.57690549862 +      536.80451209540*t);
+   earth_x_0+=     0.00000005777 * cos( 6.10164059960 +      955.59974160860*t);
+   earth_x_0+=     0.00000006393 * cos( 3.97543391449 +    24072.92146977640*t);
+   earth_x_0+=     0.00000006927 * cos( 0.71517687323 +      522.57741809380*t);
+   earth_x_0+=     0.00000007793 * cos( 4.30411864972 +    22003.91463486980*t);
+   earth_x_0+=     0.00000006346 * cos( 3.50687769075 +     5230.80746680300*t);
+   earth_x_0+=     0.00000007534 * cos( 3.51457135016 +    18422.62935909819*t);
+   earth_x_0+=     0.00000006953 * cos( 2.74384548965 +      155.42039943420*t);
+   earth_x_0+=     0.00000006224 * cos( 2.24416723229 +    18451.07854656599*t);
+   earth_x_0+=     0.00000005830 * cos( 2.32860514691 +      639.89728631400*t);
+   earth_x_0+=     0.00000005762 * cos( 3.49866583471 +     1551.04522264800*t);
+   earth_x_0+=     0.00000006251 * cos( 4.45479090841 +     6288.59877429880*t);
+   earth_x_0+=     0.00000006251 * cos( 2.19374266424 +     6277.55292568400*t);
+   earth_x_0+=     0.00000006278 * cos( 0.64265846614 +     5216.58037280140*t);
+   earth_x_0+=     0.00000006405 * cos( 2.25334669847 +     5650.29211067820*t);
+   earth_x_0+=     0.00000006545 * cos( 1.98229464447 +      103.09277421860*t);
+   earth_x_0+=     0.00000004737 * cos( 3.74210997787 +    14314.16811304980*t);
+   earth_x_0+=     0.00000004646 * cos( 1.33492108974 +      553.56940284240*t);
+   earth_x_0+=     0.00000006263 * cos( 4.39529396163 +    26087.90314157420*t);
+   earth_x_0+=     0.00000006196 * cos( 2.58704269072 +    84672.47584450469*t);
+   earth_x_0+=     0.00000006160 * cos( 4.78249950381 +   239424.39025435288*t);
+   earth_x_0+=     0.00000004988 * cos( 0.73122576373 +    77713.77146812050*t);
+   earth_x_0+=     0.00000005414 * cos( 2.99961957880 +     6179.98307577280*t);
+   earth_x_0+=     0.00000005414 * cos( 3.65424820690 +     6386.16862421000*t);
+   earth_x_0+=     0.00000004863 * cos( 3.76722340488 +    90279.92316810328*t);
+   earth_x_0+=     0.00000005711 * cos( 4.32714492993 +     6915.85958930460*t);
+   earth_x_0+=     0.00000004269 * cos( 2.38367994999 +     7875.67186362420*t);
+   earth_x_0+=     0.00000004017 * cos( 3.78580907612 +     6254.62666252360*t);
+   earth_x_0+=     0.00000004067 * cos( 2.42033776129 +     3634.62102451840*t);
+   earth_x_0+=     0.00000003993 * cos( 2.86272124248 +     6311.52503745920*t);
+   earth_x_0+=     0.00000005305 * cos( 3.34359172920 +    25158.60171976540*t);
+   earth_x_0+=     0.00000004804 * cos( 2.82827032932 +     5760.49843189760*t);
+   earth_x_0+=     0.00000004349 * cos( 5.62752376053 +     7238.67559160000*t);
+   earth_x_0+=     0.00000004119 * cos( 3.29663921982 +     6709.67404086740*t);
+   earth_x_0+=     0.00000004024 * cos( 0.55473527406 +    11015.10647733480*t);
+   earth_x_0+=     0.00000004409 * cos( 3.80563977803 +     6805.65326808520*t);
+   earth_x_0+=     0.00000003348 * cos( 3.54433469125 +    12592.45001978260*t);
+   earth_x_0+=     0.00000003172 * cos( 5.19146555016 +     1066.49547719000*t);
+   earth_x_0+=     0.00000003532 * cos( 3.92898013298 +     9917.69687450980*t);
+   earth_x_0+=     0.00000003580 * cos( 4.19680692038 +    11926.25441366880*t);
+   earth_x_0+=     0.00000003336 * cos( 1.79793376416 +    10575.40668294180*t);
+   earth_x_0+=     0.00000003149 * cos( 6.21487368875 +     4292.33083295040*t);
+   earth_x_0+=     0.00000003899 * cos( 5.53698997029 +       74.78159856730*t);
+   earth_x_0+=     0.00000002979 * cos( 6.20067573402 +    12566.21901028560*t);
+   earth_x_0+=     0.00000003981 * cos( 1.19685195313 +    13341.67431130680*t);
+   earth_x_0+=     0.00000003025 * cos( 3.18620101831 +    18052.92954315780*t);
+   earth_x_0+=     0.00000003624 * cos( 5.60981856994 +     6208.29425142410*t);
+   earth_x_0+=     0.00000003581 * cos( 1.03995231220 +     6357.85744855870*t);
+   earth_x_0+=     0.00000003698 * cos( 5.22105764023 +     5966.68398033480*t);
+   earth_x_0+=     0.00000002793 * cos( 4.13369432245 +        3.52311834900*t);
+   earth_x_0+=     0.00000003626 * cos( 1.42632163253 +     6599.46771964800*t);
+   earth_x_0+=     0.00000003625 * cos( 3.00192353938 +     6019.99192661860*t);
+   earth_x_0+=     0.00000002611 * cos( 2.43263787635 +     6702.56049386660*t);
+   earth_x_0+=     0.00000002535 * cos( 1.58310693201 +    31415.37924995700*t);
+   earth_x_0+=     0.00000003520 * cos( 5.37986364075 +   250570.67585721909*t);
+   earth_x_0+=     0.00000002593 * cos( 4.91225714775 +    16730.46368959580*t);
+   earth_x_0+=     0.00000003434 * cos( 3.64656881585 +     6546.15977336420*t);
+   earth_x_0+=     0.00000002788 * cos( 3.30835728306 +    12721.57209941700*t);
+   earth_x_0+=     0.00000003383 * cos( 1.74748489437 +    14945.31617355440*t);
+   earth_x_0+=     0.00000002653 * cos( 4.42220751136 +     5863.59120611620*t);
+   earth_x_0+=     0.00000003049 * cos( 0.02765893726 +     4732.03062734340*t);
+   earth_x_0+=     0.00000002378 * cos( 3.62936279948 +     2388.89402044920*t);
+   earth_x_0+=     0.00000002979 * cos( 0.36159339669 +    12566.08438968000*t);
+   earth_x_0+=     0.00000002716 * cos( 3.06935305016 +    12029.34718788740*t);
+   earth_x_0+=     0.00000002216 * cos( 5.26633529197 +    13362.44970679920*t);
+   earth_x_0+=     0.00000003014 * cos( 4.76640877174 +    29088.81141598500*t);
+   earth_x_0+=     0.00000002742 * cos( 5.06666852445 +    13916.01910964160*t);
+   earth_x_0+=     0.00000002265 * cos( 5.95706192557 +     6076.89030155420*t);
+   earth_x_0+=     0.00000002529 * cos( 5.94709040579 +    12043.57428188900*t);
+   earth_x_0+=     0.00000002015 * cos( 2.48432243402 +     6262.72053059260*t);
+   earth_x_0+=     0.00000002218 * cos( 0.69324011482 +     6489.26139842860*t);
+   earth_x_0+=     0.00000002005 * cos( 4.16423889100 +     6303.43116939020*t);
+   earth_x_0+=     0.00000002009 * cos( 4.09759324685 +     4701.11650170840*t);
+   earth_x_0+=     0.00000002105 * cos( 4.63993587835 +     5746.27133789600*t);
+   earth_x_0+=     0.00000001937 * cos( 0.22112768064 +    13521.75144159140*t);
+   earth_x_0+=     0.00000002437 * cos( 6.12896014686 +      149.56319713460*t);
+   earth_x_0+=     0.00000002602 * cos( 4.28001549113 +     4590.91018048900*t);
+   earth_x_0+=     0.00000001893 * cos( 5.66795892365 +     6286.95718534940*t);
+   earth_x_0+=     0.00000001893 * cos( 0.98057464900 +     6279.19451463340*t);
+   earth_x_0+=     0.00000002156 * cos( 2.00898711721 +    18139.29450141590*t);
+   earth_x_0+=     0.00000002388 * cos( 0.13031678718 +     1162.47470440780*t);
+   earth_x_0+=     0.00000001962 * cos( 0.25060934155 +     4686.88940770680*t);
+   earth_x_0+=     0.00000002337 * cos( 4.40745417534 +    10021.83728009940*t);
+   earth_x_0+=     0.00000001884 * cos( 5.27011523204 +      242.72860397400*t);
+   earth_x_0+=     0.00000002436 * cos( 4.40772141072 +    95143.13292097810*t);
+   earth_x_0+=     0.00000001759 * cos( 0.19888470443 +    13517.87010623340*t);
+   earth_x_0+=     0.00000002163 * cos( 1.41797947214 +     6037.24420376200*t);
+   earth_x_0+=     0.00000002167 * cos( 4.72612563566 +    23013.53953958720*t);
+   earth_x_0+=     0.00000002250 * cos( 2.39728930597 +    12669.24447420140*t);
+   earth_x_0+=     0.00000002052 * cos( 5.23006568302 +     6528.90749622080*t);
+   earth_x_0+=     0.00000001774 * cos( 2.08484497348 +     6819.88036208680*t);
+   earth_x_0+=     0.00000001842 * cos( 0.34675193374 +    65147.61976813770*t);
+   earth_x_0+=     0.00000001549 * cos( 0.71612589777 +       20.77539549240*t);
+   earth_x_0+=     0.00000001543 * cos( 5.96098287839 +     6272.43918464160*t);
+   earth_x_0+=     0.00000001543 * cos( 0.68755069426 +     6293.71251534120*t);
+   earth_x_0+=     0.00000001727 * cos( 4.50674659690 +    17654.78053974960*t);
+   earth_x_0+=     0.00000001436 * cos( 2.99734831658 +      110.20632121940*t);
+   earth_x_0+=     0.00000001427 * cos( 1.68851809181 +    13119.72110282519*t);
+   earth_x_0+=     0.00000001395 * cos( 3.25347392533 +    10177.25767953360*t);
+   earth_x_0+=     0.00000001446 * cos( 2.26563038586 +     1052.26838318840*t);
+   earth_x_0+=     0.00000001819 * cos( 3.18848907826 +    18875.52586977400*t);
+   earth_x_0+=     0.00000001740 * cos( 5.48710781556 +    19651.04848109800*t);
+   earth_x_0+=     0.00000001794 * cos( 5.74176427744 +    25934.12433108940*t);
+   earth_x_0+=     0.00000001800 * cos( 6.18908014275 +    33019.02111220460*t);
+   earth_x_0+=     0.00000001488 * cos( 5.67443981584 +       76.26607127560*t);
+   earth_x_0+=     0.00000001349 * cos( 5.75097793699 +     6016.46880826960*t);
+   earth_x_0+=     0.00000001477 * cos( 4.14004787537 +    12779.45079542080*t);
+   earth_x_0+=     0.00000001251 * cos( 1.50045769848 +     3894.18182954220*t);
+   earth_x_0+=     0.00000001492 * cos( 6.22975090570 +     9411.46461508720*t);
+   earth_x_0+=     0.00000001471 * cos( 5.75783391755 +    18635.92845453620*t);
+   earth_x_0+=     0.00000001672 * cos( 5.43726223249 +    16200.77272450120*t);
+   earth_x_0+=     0.00000001263 * cos( 0.91728971350 +    12573.26524698360*t);
+   earth_x_0+=     0.00000001366 * cos( 1.09160445123 +     2379.16447357160*t);
+   earth_x_0+=     0.00000001272 * cos( 0.89738866811 +     6549.68289171320*t);
+   earth_x_0+=     0.00000001274 * cos( 5.34253525843 +      632.78373931320*t);
+   earth_x_0+=     0.00000001156 * cos( 1.38847609016 +    11499.65622279280*t);
+   earth_x_0+=     0.00000001543 * cos( 4.26237609038 +    18209.33026366019*t);
+   earth_x_0+=     0.00000001460 * cos( 4.97233669659 +      245.83164622940*t);
+   earth_x_0+=     0.00000001253 * cos( 1.41164782000 +     6993.00889854970*t);
+   earth_x_0+=     0.00000001455 * cos( 5.29534379992 +       38.13303563780*t);
+   earth_x_0+=     0.00000001228 * cos( 2.83750345253 +    24356.78078864160*t);
+   earth_x_0+=     0.00000001028 * cos( 2.51030023881 +     6112.40297937220*t);
+   earth_x_0+=     0.00000001268 * cos( 4.67217109478 +     5429.87946823940*t);
+   earth_x_0+=     0.00000001312 * cos( 1.10091569645 +    12491.37010141550*t);
+   earth_x_0+=     0.00000000992 * cos( 4.13817717100 +     6453.74872061060*t);
+   earth_x_0+=     0.00000000988 * cos( 0.98006619973 +     8273.82086703240*t);
+   earth_x_0+=     0.00000001191 * cos( 3.21845317343 +     6276.02961329339*t);
+   earth_x_0+=     0.00000001191 * cos( 3.43008039922 +     6290.12208668940*t);
+   earth_x_0+=     0.00000000978 * cos( 4.73852763851 +     6247.04798331399*t);
+   earth_x_0+=     0.00000000971 * cos( 1.91000718426 +     6319.10371666880*t);
+   earth_x_0+=     0.00000000962 * cos( 2.63568544853 +    12565.17137891460*t);
+   earth_x_0+=     0.00000001025 * cos( 2.21768401152 +    12545.37630449040*t);
+   earth_x_0+=     0.00000001146 * cos( 3.85865009647 +      206.18554843720*t);
+   earth_x_0+=     0.00000000944 * cos( 0.78473810598 +     6306.95428773920*t);
+   earth_x_0+=     0.00000000944 * cos( 5.86379546667 +     6259.19741224360*t);
+   earth_x_0+=     0.00000001092 * cos( 3.92002105501 +    17796.95916678580*t);
+   earth_x_0+=     0.00000000886 * cos( 5.23132430543 +     6172.86952877200*t);
+   earth_x_0+=     0.00000000859 * cos( 1.48130342803 +     1692.16566950240*t);
+   earth_x_0+=     0.00000001211 * cos( 4.04154480615 +    82576.98122099529*t);
+   earth_x_0+=     0.00000001132 * cos( 0.96198039943 +    17782.73207278420*t);
+   earth_x_0+=     0.00000000942 * cos( 1.42062620155 +     6267.82337880639*t);
+   earth_x_0+=     0.00000000942 * cos( 5.22790737110 +     6298.32832117640*t);
+   earth_x_0+=     0.00000001045 * cos( 0.46337213067 +    15508.61512327440*t);
+   earth_x_0+=     0.00000000850 * cos( 1.45204280805 +     6393.28217121080*t);
+   earth_x_0+=     0.00000000895 * cos( 2.40085098648 +     3903.91137641980*t);
+   earth_x_0+=     0.00000001016 * cos( 0.71504248238 +      220.41264243880*t);
+   earth_x_0+=     0.00000000853 * cos( 4.42455068314 +      949.17560896980*t);
+   earth_x_0+=     0.00000000917 * cos( 2.47015034331 +   149854.40013480789*t);
+   earth_x_0+=     0.00000001085 * cos( 1.32862144392 +    36949.23080842420*t);
+   earth_x_0+=     0.00000000779 * cos( 3.58339451743 +     2648.45482547300*t);
+   earth_x_0+=     0.00000001000 * cos( 0.24880432000 +    11403.67699557500*t);
+   earth_x_0+=     0.00000000962 * cos( 3.92658368218 +    12567.13202105100*t);
+   earth_x_0+=     0.00000000770 * cos( 6.20173754399 +     4804.20927592700*t);
+   earth_x_0+=     0.00000000986 * cos( 4.62897469531 +    22805.73556599360*t);
+   earth_x_0+=     0.00000000931 * cos( 2.22558469426 +     2787.04302385740*t);
+   earth_x_0+=     0.00000000743 * cos( 3.79903933428 +     7834.12107263940*t);
+   earth_x_0+=     0.00000000905 * cos( 3.81538816339 +     6284.56032269970*t);
+   earth_x_0+=     0.00000000905 * cos( 2.83314540925 +     6281.59137728310*t);
+   earth_x_0+=     0.00000000715 * cos( 5.50312651042 +    17256.63153634140*t);
+   earth_x_0+=     0.00000000798 * cos( 2.62501077543 +    12410.73130054860*t);
+   earth_x_0+=     0.00000000828 * cos( 1.54060095214 +    10969.96525769820*t);
+   earth_x_0+=     0.00000000821 * cos( 0.52122150189 +    12416.58850284820*t);
+   earth_x_0+=     0.00000000826 * cos( 4.78071938299 +    10984.19235169980*t);
+   earth_x_0+=     0.00000000754 * cos( 3.82898538792 +     4061.21921539440*t);
+   earth_x_0+=     0.00000000616 * cos( 5.70479516739 +     1596.18644228460*t);
+   earth_x_0+=     0.00000000723 * cos( 3.06755247513 +   161000.68573767410*t);
+   earth_x_0+=     0.00000000782 * cos( 0.87676753080 +     4907.30205014560*t);
+   earth_x_0+=     0.00000000707 * cos( 3.28533306159 +     5849.36411211460*t);
+   earth_x_0+=     0.00000000783 * cos( 3.43481557690 +      846.08283475120*t);
+   earth_x_0+=     0.00000000623 * cos( 2.94730498294 +     5333.90024102160*t);
+   earth_x_0+=     0.00000000696 * cos( 4.89623582044 +    11712.95531823080*t);
+   earth_x_0+=     0.00000000616 * cos( 1.47727208161 +     5573.14280143310*t);
+   earth_x_0+=     0.00000000595 * cos( 1.30546777768 +    12323.42309600880*t);
+   earth_x_0+=     0.00000000611 * cos( 1.31321200457 +    11933.36796066960*t);
+   earth_x_0+=     0.00000000611 * cos( 1.53155679450 +      853.19638175200*t);
+   earth_x_0+=     0.00000000570 * cos( 4.54665961586 +    12586.92709547520*t);
+   earth_x_0+=     0.00000000778 * cos( 1.57590278073 +     8662.24032356300*t);
+   earth_x_0+=     0.00000000597 * cos( 4.69316038420 +     9779.10867612540*t);
+   earth_x_0+=     0.00000000536 * cos( 6.06660590951 +    18852.75066832320*t);
+   earth_x_0+=     0.00000000532 * cos( 4.58762153412 +    18842.11400297339*t);
+   earth_x_0+=     0.00000000599 * cos( 2.11793250999 +     6151.53388830500*t);
+   earth_x_0+=     0.00000000540 * cos( 3.40440707564 +     6279.78949257360*t);
+   earth_x_0+=     0.00000000540 * cos( 3.24412649701 +     6286.36220740920*t);
+   earth_x_0+=     0.00000000621 * cos( 1.14533482552 +    52175.80628314840*t);
+   earth_x_0+=     0.00000000513 * cos( 3.70291605445 +    12455.94537876340*t);
+   earth_x_0+=     0.00000000623 * cos( 5.52546598759 +     5326.78669402080*t);
+   earth_x_0+=     0.00000000649 * cos( 4.90568768794 +    11764.33076885900*t);
+   earth_x_0+=     0.00000000582 * cos( 4.53092407532 +     6414.61781167780*t);
+   earth_x_0+=     0.00000000610 * cos( 3.94320832220 +    10419.98628350760*t);
+   earth_x_0+=     0.00000000679 * cos( 0.26045226696 +     6133.51265285680*t);
+   earth_x_0+=     0.00000000531 * cos( 3.00684541494 +     3496.03282613400*t);
+   earth_x_0+=     0.00000000536 * cos( 3.70224967645 +     7232.25145896120*t);
+   earth_x_0+=     0.00000000510 * cos( 4.38198402315 +    11513.88331679440*t);
+   earth_x_0+=     0.00000000517 * cos( 2.93096521231 +      433.71173787680*t);
+   earth_x_0+=     0.00000000643 * cos( 6.06382589060 +    26735.94526221320*t);
+   earth_x_0+=     0.00000000665 * cos( 2.75085842376 +    40879.44050464380*t);
+   earth_x_0+=     0.00000000549 * cos( 0.63911798007 +    22483.84857449259*t);
+   earth_x_0+=     0.00000000479 * cos( 1.33580595339 +     5547.19933645960*t);
+   earth_x_0+=     0.00000000564 * cos( 3.92025647347 +     4171.42553661380*t);
+   earth_x_0+=     0.00000000462 * cos( 1.96611775944 +     3104.93005942380*t);
+   earth_x_0+=     0.00000000551 * cos( 4.03874430699 +     6503.48849243020*t);
+   earth_x_0+=     0.00000000535 * cos( 6.12277226234 +   173566.83743765688*t);
+   earth_x_0+=     0.00000000545 * cos( 2.61441075819 +     6062.66320755260*t);
+   earth_x_0+=     0.00000000477 * cos( 3.50657951584 +     9388.00590941520*t);
+   earth_x_0+=     0.00000000448 * cos( 5.09815034585 +     4274.51831083240*t);
+   earth_x_0+=     0.00000000532 * cos( 0.22623731981 +      283.85931886520*t);
+   earth_x_0+=     0.00000000527 * cos( 6.24827316267 +    24602.61243487099*t);
+   earth_x_0+=     0.00000000502 * cos( 0.97407194728 +    12489.88562870720*t);
+   earth_x_0+=     0.00000000603 * cos( 0.85487083482 +   322711.30452390644*t);
+   earth_x_0+=     0.00000000493 * cos( 6.11920290350 +    29296.61538957860*t);
+   earth_x_0+=     0.00000000530 * cos( 0.43491203890 +    18845.70443162520*t);
+   earth_x_0+=     0.00000000431 * cos( 5.31145947560 +     7018.95236352320*t);
+   earth_x_0+=     0.00000000559 * cos( 1.71609589414 +     1581.95934828300*t);
+   earth_x_0+=     0.00000000584 * cos( 0.07135689258 +    72850.56171524568*t);
+   earth_x_0+=     0.00000000478 * cos( 1.85248622079 +       14.22709400160*t);
+   earth_x_0+=     0.00000000569 * cos( 2.95236791425 +    97238.62754448749*t);
+   earth_x_0+=     0.00000000405 * cos( 5.39493918596 +    10873.98603048040*t);
+   earth_x_0+=     0.00000000401 * cos( 4.50443000854 +    17267.26820169119*t);
+   earth_x_0+=     0.00000000389 * cos( 3.43047350256 +     6244.94281435360*t);
+   earth_x_0+=     0.00000000389 * cos( 3.21806007009 +     6321.20888562920*t);
+   earth_x_0+=     0.00000000388 * cos( 1.13963224225 +     6297.30294399300*t);
+   earth_x_0+=     0.00000000388 * cos( 5.50890133040 +     6268.84875598980*t);
+   earth_x_0+=     0.00000000515 * cos( 1.67075193318 +    12320.32005375340*t);
+   earth_x_0+=     0.00000000393 * cos( 4.37184659923 +    14919.01785375460*t);
+   earth_x_0+=     0.00000000394 * cos( 6.10411050787 +     4157.19844261220*t);
+   earth_x_0+=     0.00000000491 * cos( 5.96740803508 +     1478.86657406440*t);
+   earth_x_0+=     0.00000000505 * cos( 4.66096327599 +    17157.06188047180*t);
+   earth_x_0+=     0.00000000528 * cos( 5.62841514056 +      135.06508003540*t);
+   earth_x_0+=     0.00000000439 * cos( 4.82558161657 +   233141.31440436149*t);
+   earth_x_0+=     0.00000000417 * cos( 3.74443512826 +    13760.59871020740*t);
+   earth_x_0+=     0.00000000419 * cos( 5.08318905121 +    15671.08175940660*t);
+   earth_x_0+=     0.00000000393 * cos( 1.56711153414 +     7872.14874527520*t);
+   earth_x_0+=     0.00000000392 * cos( 1.30305906278 +      419.48464387520*t);
+   earth_x_0+=     0.00000000365 * cos( 2.11020607854 +     6922.97313630540*t);
+   earth_x_0+=     0.00000000419 * cos( 0.62417554081 +     5642.19824260920*t);
+   earth_x_0+=     0.00000000387 * cos( 0.68719040069 +    17253.04110768959*t);
+   earth_x_0+=     0.00000000371 * cos( 5.76233828697 +    13625.53363017200*t);
+   earth_x_0+=     0.00000000379 * cos( 0.81831291068 +     3185.19202726560*t);
+   earth_x_0+=     0.00000000499 * cos( 2.72563350590 +   161710.61878623239*t);
+   earth_x_0+=     0.00000000414 * cos( 3.32204347851 +    22779.43724619380*t);
+   earth_x_0+=     0.00000000485 * cos( 1.35444198321 +    12528.01866434500*t);
+   earth_x_0+=     0.00000000358 * cos( 3.15737999738 +     6271.34649715540*t);
+   earth_x_0+=     0.00000000358 * cos( 3.49115357527 +     6294.80520282740*t);
+   earth_x_0+=     0.00000000351 * cos( 2.21887707287 +    11616.97609101300*t);
+   earth_x_0+=     0.00000000387 * cos( 0.49044025161 +      956.28915597060*t);
+   earth_x_0+=     0.00000000451 * cos( 0.91836677580 +    29864.33402730900*t);
+   earth_x_0+=     0.00000000421 * cos( 4.82745120825 +     7445.55055439920*t);
+   earth_x_0+=     0.00000000327 * cos( 3.95232223910 +    20995.39296644940*t);
+   earth_x_0+=     0.00000000336 * cos( 1.41330667750 +     6279.14369672830*t);
+   earth_x_0+=     0.00000000336 * cos( 5.23522689515 +     6287.00800325450*t);
+   earth_x_0+=     0.00000000358 * cos( 0.81490874525 +    21393.54196985760*t);
+   earth_x_0+=     0.00000000450 * cos( 0.23149923655 +     1385.56157495110*t);
+   earth_x_0+=     0.00000000320 * cos( 0.03172757569 +    16858.48253293320*t);
+   earth_x_0+=     0.00000000348 * cos( 4.27565702309 +        3.93215326310*t);
+   earth_x_0+=     0.00000000363 * cos( 4.29069778349 +     7335.34423317980*t);
+   earth_x_0+=     0.00000000330 * cos( 3.68113867113 +       38.02767263580*t);
+   earth_x_0+=     0.00000000395 * cos( 5.39428043096 +     5905.70224207560*t);
+   earth_x_0+=     0.00000000414 * cos( 4.19402575890 +    44809.65020086340*t);
+   earth_x_0+=     0.00000000300 * cos( 3.87083875051 +      640.87760738220*t);
+   earth_x_0+=     0.00000000304 * cos( 0.15986036979 +     5749.86176654780*t);
+   earth_x_0+=     0.00000000351 * cos( 5.93395455459 +    12345.73905754400*t);
+   earth_x_0+=     0.00000000307 * cos( 5.30261033229 +     5120.60114558360*t);
+   earth_x_0+=     0.00000000304 * cos( 4.93657040735 +    29826.30635467320*t);
+   earth_x_0+=     0.00000000297 * cos( 6.01664889580 +    20597.24396304120*t);
+   earth_x_0+=     0.00000000395 * cos( 0.07726910331 +     4176.04134244900*t);
+   earth_x_0+=     0.00000000320 * cos( 2.49799253183 +    10818.13528691580*t);
+   earth_x_0+=     0.00000000365 * cos( 1.25425297875 +     6660.44945790720*t);
+   earth_x_0+=     0.00000000339 * cos( 6.10828911901 +     6923.95345737360*t);
+   earth_x_0+=     0.00000000274 * cos( 1.11162071728 +    26709.64694241340*t);
+   earth_x_0+=     0.00000000386 * cos( 2.29044623019 +     2636.72547263700*t);
+   earth_x_0+=     0.00000000367 * cos( 1.18722914434 +    30666.15495843280*t);
+   earth_x_0+=     0.00000000268 * cos( 4.95835137530 +    10454.50138660520*t);
+   earth_x_0+=     0.00000000265 * cos( 5.11010291991 +    19378.91851506880*t);
+   earth_x_0+=     0.00000000325 * cos( 2.49277998223 +    12359.96615154560*t);
+   earth_x_0+=     0.00000000326 * cos( 4.24702233173 +    16522.65971600220*t);
+   earth_x_0+=     0.00000000269 * cos( 0.20537129126 +     6816.28993343500*t);
+   earth_x_0+=     0.00000000298 * cos( 5.15059518785 +    10344.29506538580*t);
+   earth_x_0+=     0.00000000344 * cos( 3.51559745333 +     6147.45052498140*t);
+   earth_x_0+=     0.00000000255 * cos( 4.94642831875 +     1375.77379984580*t);
+   earth_x_0+=     0.00000000283 * cos( 2.57877323959 +      316.39186965660*t);
+   earth_x_0+=     0.00000000334 * cos( 3.13303102229 +     6418.70117500140*t);
+   earth_x_0+=     0.00000000316 * cos( 5.54205366730 +     6055.59971720240*t);
+   earth_x_0+=     0.00000000316 * cos( 0.36194124996 +    21954.15760939799*t);
+   earth_x_0+=     0.00000000243 * cos( 4.10172883526 +     6310.47740608820*t);
+   earth_x_0+=     0.00000000243 * cos( 2.54680473739 +     6255.67429389460*t);
+   earth_x_0+=     0.00000000314 * cos( 4.99359717889 +    17996.03116822220*t);
+   earth_x_0+=     0.00000000296 * cos( 0.07899534657 +     6129.29703950660*t);
+   earth_x_0+=     0.00000000301 * cos( 1.10648000281 +     6510.55198278040*t);
+   earth_x_0+=     0.00000000329 * cos( 1.15676558971 +      377.37360791580*t);
+   earth_x_0+=     0.00000000267 * cos( 3.42007125342 +    24705.70520908959*t);
+   earth_x_0+=     0.00000000285 * cos( 2.50494681053 +     5469.52556603160*t);
+   earth_x_0+=     0.00000000311 * cos( 3.23911937497 +    11720.06886523160*t);
+   earth_x_0+=     0.00000000286 * cos( 0.28647339668 +     6436.85466047620*t);
+   earth_x_0+=     0.00000000273 * cos( 4.29648399498 +    28286.99048486120*t);
+   earth_x_0+=     0.00000000319 * cos( 4.71062528834 +   169379.25621117488*t);
+   earth_x_0+=     0.00000000229 * cos( 1.01384122818 +     1265.56747862640*t);
+   earth_x_0+=     0.00000000234 * cos( 3.36893910039 +     4487.81740627040*t);
+   earth_x_0+=     0.00000000252 * cos( 0.51301411030 +     4377.61108505100*t);
+   earth_x_0+=     0.00000000214 * cos( 2.09487711967 +     2221.85663459700*t);
+   earth_x_0+=     0.00000000213 * cos( 3.79431896643 +    30355.99731976780*t);
+   earth_x_0+=     0.00000000263 * cos( 6.11317571912 +     9380.95967271720*t);
+   earth_x_0+=     0.00000000239 * cos( 1.33165763648 +    10440.27429260360*t);
+   earth_x_0+=     0.00000000289 * cos( 0.55454156939 +     6206.80977871580*t);
+   earth_x_0+=     0.00000000259 * cos( 3.43127382595 +     6282.83203250789*t);
+   earth_x_0+=     0.00000000259 * cos( 3.21725974670 +     6283.31966747490*t);
+   earth_x_0+=     0.00000000243 * cos( 4.14346390349 +     7096.62613395120*t);
+   earth_x_0+=     0.00000000246 * cos( 5.83734028861 +     8671.96987044060*t);
+   earth_x_0+=     0.00000000199 * cos( 0.93488663037 +    11216.28429032400*t);
+   earth_x_0+=     0.00000000265 * cos( 5.61985831637 +    48739.85989708300*t);
+   earth_x_0+=     0.00000000207 * cos( 3.06628694399 +    12562.56127133100*t);
+   earth_x_0+=     0.00000000212 * cos( 1.75471606720 +     7129.15868474260*t);
+   earth_x_0+=     0.00000000269 * cos( 1.52668826951 +    70755.06709173629*t);
+   earth_x_0+=     0.00000000225 * cos( 4.01656879083 +     5636.06501667660*t);
+   earth_x_0+=     0.00000000192 * cos( 2.04748658291 +    17271.88400752640*t);
+   earth_x_0+=     0.00000000240 * cos( 5.09605084016 +    15906.76412668260*t);
+   earth_x_0+=     0.00000000187 * cos( 0.15192842664 +    11609.86254401220*t);
+   earth_x_0+=     0.00000000256 * cos( 2.30143585875 +      224.34479570190*t);
+   earth_x_0+=     0.00000000250 * cos( 4.24183292428 +     5934.15142954340*t);
+   earth_x_0+=     0.00000000220 * cos( 0.73345029623 +    11087.28512591840*t);
+   earth_x_0+=     0.00000000175 * cos( 5.41174539523 +    11614.43329373220*t);
+   earth_x_0+=     0.00000000175 * cos( 4.58212579300 +     3744.82734573780*t);
+   earth_x_0+=     0.00000000241 * cos( 0.28788321480 +     8390.11035753380*t);
+   earth_x_0+=     0.00000000234 * cos( 2.61192358068 +    34596.36465465240*t);
+   earth_x_0+=     0.00000000174 * cos( 2.85207900960 +      263.08392337280*t);
+   earth_x_0+=     0.00000000233 * cos( 2.02657287745 +        1.48447270830*t);
+   earth_x_0+=     0.00000000190 * cos( 6.26917436708 +      394.62588505920*t);
+   earth_x_0+=     0.00000000187 * cos( 5.58825382870 +      401.67212175720*t);
+   earth_x_0+=     0.00000000178 * cos( 3.72578587318 +    12132.43996210600*t);
+   earth_x_0+=     0.00000000230 * cos( 2.40692341132 +     6632.00027043940*t);
+   earth_x_0+=     0.00000000177 * cos( 4.73089168144 +     4797.09572892620*t);
+   earth_x_0+=     0.00000000167 * cos( 3.97388034335 +     5714.25397596400*t);
+   earth_x_0+=     0.00000000160 * cos( 2.37187295106 +      742.99006053260*t);
+   earth_x_0+=     0.00000000214 * cos( 4.34116392416 +    24492.40611365159*t);
+   earth_x_0+=     0.00000000217 * cos( 4.52620284574 +     3097.88382272579*t);
+   earth_x_0+=     0.00000000172 * cos( 5.38973312367 +     6332.83287546320*t);
+   earth_x_0+=     0.00000000172 * cos( 1.25880044898 +     6233.31882451960*t);
+   earth_x_0+=     0.00000000221 * cos( 4.98998061605 +      266.60704172180*t);
+   earth_x_0+=     0.00000000159 * cos( 3.90246086071 +    13198.93543929600*t);
+   earth_x_0+=     0.00000000207 * cos( 3.49598218672 +    12569.74212863460*t);
+   earth_x_0+=     0.00000000176 * cos( 4.05199455874 +    12012.58229714040*t);
+   earth_x_0+=     0.00000000157 * cos( 5.43302098132 +    16840.67001081519*t);
+   earth_x_0+=     0.00000000170 * cos( 0.63813817825 +     6272.03014972750*t);
+   earth_x_0+=     0.00000000170 * cos( 6.01039539440 +     6294.12155025530*t);
+   earth_x_0+=     0.00000000213 * cos( 4.21609963122 +    19804.82729158280*t);
+   earth_x_0+=     0.00000000167 * cos( 5.54070880349 +    17686.75284556640*t);
+   earth_x_0+=     0.00000000187 * cos( 5.70099962642 +    20452.86941222180*t);
+   earth_x_0+=     0.00000000147 * cos( 4.99188388936 +    16737.57723659660*t);
+   earth_x_0+=     0.00000000148 * cos( 2.67465318849 +     6851.89772401880*t);
+   earth_x_0+=     0.00000000146 * cos( 0.23801153618 +    17370.36097590980*t);
+   earth_x_0+=     0.00000000143 * cos( 0.85202449839 +     5618.31980486140*t);
+   earth_x_0+=     0.00000000193 * cos( 3.57442994352 +    33326.57873317420*t);
+   earth_x_0+=     0.00000000140 * cos( 1.22261707467 +     2008.55753915900*t);
+   earth_x_0+=     0.00000000155 * cos( 0.27598891874 +    20199.09495963300*t);
+   earth_x_0+=     0.00000000182 * cos( 5.69212378196 +    28766.92442448400*t);
+   earth_x_0+=     0.00000000184 * cos( 3.48483358167 +     5540.08578945880*t);
+   earth_x_0+=     0.00000000162 * cos( 4.78754155685 +    12551.92460598120*t);
+   earth_x_0+=     0.00000000143 * cos( 2.32151913243 +     5010.39482436420*t);
+   earth_x_0+=     0.00000000190 * cos( 0.66924538276 +    23581.25817731760*t);
+   earth_x_0+=     0.00000000135 * cos( 5.34929799142 +      735.87651353180*t);
+   earth_x_0+=     0.00000000135 * cos( 0.07815962516 +    12431.08661994740*t);
+   earth_x_0+=     0.00000000134 * cos( 3.03203004875 +    11925.27409260060*t);
+   earth_x_0+=     0.00000000179 * cos( 1.07503153366 +    16460.33352952499*t);
+   earth_x_0+=     0.00000000182 * cos( 3.66470209346 +       28.44918746780*t);
+   earth_x_0+=     0.00000000146 * cos( 1.55848605346 +    12701.21678001820*t);
+   earth_x_0+=     0.00000000137 * cos( 2.60133025404 +       11.04570026390*t);
+   earth_x_0+=     0.00000000167 * cos( 1.26986508714 +    17473.45375012840*t);
+   earth_x_0+=     0.00000000170 * cos( 0.76594212114 +    52670.06959330260*t);
+   earth_x_0+=     0.00000000124 * cos( 5.79650915048 +     6947.83189512140*t);
+   earth_x_0+=     0.00000000135 * cos( 3.38698681678 +     3531.52825029980*t);
+   earth_x_0+=     0.00000000172 * cos( 4.94248686856 +   167959.39011405829*t);
+   earth_x_0+=     0.00000000122 * cos( 5.69264255527 +     3583.34103067380*t);
+   earth_x_0+=     0.00000000170 * cos( 1.45223415923 +   333857.59012677266*t);
+   earth_x_0+=     0.00000000155 * cos( 4.05744056379 +     6058.73105428950*t);
+   earth_x_0+=     0.00000000120 * cos( 1.19140907580 +    12808.88030395680*t);
+   earth_x_0+=     0.00000000141 * cos( 5.52534979149 +   162420.55183479069*t);
+   earth_x_0+=     0.00000000124 * cos( 1.37398060221 +     8982.81066930900*t);
+   earth_x_0+=     0.00000000120 * cos( 5.84076902832 +    12528.12402734700*t);
+   earth_x_0+=     0.00000000119 * cos( 1.19695856457 +    19247.37655338240*t);
+   earth_x_0+=     0.00000000143 * cos( 2.13563488838 +       95.97922721780*t);
+   earth_x_0+=     0.00000000117 * cos( 0.73795877914 +    11610.55195837420*t);
+   earth_x_0+=     0.00000000154 * cos( 4.06475922223 +    38526.57435087200*t);
+   earth_x_0+=     0.00000000163 * cos( 3.11560982847 +      661.23292678100*t);
+   earth_x_0+=     0.00000000163 * cos( 4.35848910377 +     9929.42622734580*t);
+   earth_x_0+=     0.00000000115 * cos( 2.14676232906 +    19402.79695281660*t);
+   earth_x_0+=     0.00000000119 * cos( 4.28708216617 +    12249.75983032620*t);
+   earth_x_0+=     0.00000000145 * cos( 2.74250030398 +     6360.74962041940*t);
+   earth_x_0+=     0.00000000145 * cos( 3.90603326867 +     6205.40207956340*t);
+   earth_x_0+=     0.00000000159 * cos( 4.18513291774 +   228278.10465148659*t);
+   earth_x_0+=     0.00000000156 * cos( 2.57788957454 +     4307.58330413540*t);
+   earth_x_0+=     0.00000000116 * cos( 0.60146141700 +    24065.80792277559*t);
+   earth_x_0+=     0.00000000118 * cos( 1.23925579836 +     6302.74561089119*t);
+   earth_x_0+=     0.00000000118 * cos( 5.40927777429 +     6263.40608909160*t);
+   earth_x_0+=     0.00000000110 * cos( 2.33480621176 +    16062.18452611680*t);
+   earth_x_0+=     0.00000000137 * cos( 4.10438279102 +     6250.54329920000*t);
+   earth_x_0+=     0.00000000137 * cos( 2.54415078163 +     6315.60840078280*t);
+   earth_x_0+=     0.00000000111 * cos( 2.27442358208 +     6339.97422492700*t);
+   earth_x_0+=     0.00000000111 * cos( 4.37410999057 +     6226.17747505580*t);
+   earth_x_0+=     0.00000000150 * cos( 0.96205875343 +    19800.94595622480*t);
+   earth_x_0+=     0.00000000123 * cos( 2.37616666242 +    24734.15439655740*t);
+   earth_x_0+=     0.00000000107 * cos( 2.16622176546 +    12463.05892576420*t);
+   earth_x_0+=     0.00000000123 * cos( 1.00594418221 +     6305.85105144220*t);
+   earth_x_0+=     0.00000000123 * cos( 5.64258939044 +     6260.30064854060*t);
+   earth_x_0+=     0.00000000106 * cos( 1.67655703852 +    11190.37790013700*t);
+   earth_x_0+=     0.00000000113 * cos( 1.95413568499 +     2111.65031337760*t);
+   earth_x_0+=     0.00000000131 * cos( 3.61185299316 +    18415.51581209739*t);
+   earth_x_0+=     0.00000000110 * cos( 4.88667127380 +     6276.44729109139*t);
+   earth_x_0+=     0.00000000110 * cos( 1.76186229885 +     6289.70440889141*t);
+   earth_x_0+=     0.00000000103 * cos( 2.75809551865 +     6324.62664097620*t);
+   earth_x_0+=     0.00000000103 * cos( 3.89043805400 +     6241.52505900660*t);
+   earth_x_0+=     0.00000000103 * cos( 5.76083420502 +    23937.85638974100*t);
+   earth_x_0+=     0.00000000107 * cos( 1.72027333824 +    83973.83535572988*t);
+   earth_x_0+=     0.00000000107 * cos( 2.75882604296 +    84019.85928049390*t);
+   earth_x_0+=     0.00000000113 * cos( 2.65508166498 +    24336.00539314920*t);
+   earth_x_0+=     0.00000000106 * cos( 3.27680652019 +    12772.33724842000*t);
+   earth_x_0+=     0.00000000143 * cos( 1.51199781479 +     2069.00683490660*t);
+   earth_x_0+=     0.00000000143 * cos( 0.91985752497 +    18772.96147869860*t);
+   earth_x_0+=     0.00000000115 * cos( 3.78220330855 +     3641.73457151920*t);
+   earth_x_0+=     0.00000000126 * cos( 1.19294266238 +    18822.92923017440*t);
+   earth_x_0+=     0.00000000115 * cos( 4.66544469064 +    12570.08385324590*t);
+   earth_x_0+=     0.00000000098 * cos( 5.07674655291 +    16310.97904572060*t);
+   earth_x_0+=     0.00000000138 * cos( 4.55918305844 +     6209.77872413240*t);
+   earth_x_0+=     0.00000000138 * cos( 2.08935051421 +     6356.37297585040*t);
+   earth_x_0+=     0.00000000120 * cos( 4.95740945799 +     9814.60410029120*t);
+   earth_x_0+=     0.00000000135 * cos( 0.57059381593 +    18699.66435283960*t);
+   earth_x_0+=     0.00000000104 * cos( 2.15061044327 +    12560.62877567540*t);
+   earth_x_0+=     0.00000000098 * cos( 5.57514710248 +     5815.11085963699*t);
+   earth_x_0+=     0.00000000106 * cos( 1.77919890262 +    23539.70738633280*t);
+   earth_x_0+=     0.00000000114 * cos( 3.75372093481 +     6155.05700665400*t);
+   earth_x_0+=     0.00000000114 * cos( 2.89481263784 +     6411.09469332880*t);
+   earth_x_0+=     0.00000000097 * cos( 6.03825791149 +    27043.50288318280*t);
+   earth_x_0+=     0.00000000095 * cos( 5.63307798614 +    11300.58422135640*t);
+   earth_x_0+=     0.00000000110 * cos( 2.79120257993 +     6312.50535852740*t);
+   earth_x_0+=     0.00000000110 * cos( 3.85733099272 +     6253.64634145540*t);
+   earth_x_0+=     0.00000000091 * cos( 2.99457476326 +    12829.23562335560*t);
+   earth_x_0+=     0.00000000128 * cos( 5.41337210078 +     1293.48422358260*t);
+   earth_x_0+=     0.00000000096 * cos( 5.97584747761 +     2699.73481931760*t);
+   earth_x_0+=     0.00000000119 * cos( 5.98882616741 +    17892.93839400359*t);
+   earth_x_0+=     0.00000000126 * cos( 0.25750574604 +   311565.01892104017*t);
+   earth_x_0+=     0.00000000120 * cos( 2.08798412726 +    14158.74771361560*t);
+   earth_x_0+=     0.00000000100 * cos( 4.30344030661 +    16627.37091537720*t);
+   earth_x_0+=     0.00000000099 * cos( 3.61958161583 +     6432.63904712600*t);
+   earth_x_0+=     0.00000000095 * cos( 4.18452410969 +    78263.70942472259*t);
+   earth_x_0+=     0.00000000124 * cos( 0.48010020247 +   178430.04719053170*t);
+   earth_x_0+=     0.00000000089 * cos( 1.07338604948 +     6751.04084034580*t);
+   earth_x_0+=     0.00000000096 * cos( 2.64494034094 +     7349.57132718140*t);
+   earth_x_0+=     0.00000000117 * cos( 5.44811650549 +    42456.78404709160*t);
+   earth_x_0+=     0.00000000103 * cos( 4.37748903957 +    18202.21671665939*t);
+   earth_x_0+=     0.00000000099 * cos( 0.58169736416 +     6313.98997562640*t);
+   earth_x_0+=     0.00000000099 * cos( 6.06683620849 +     6252.16172435640*t);
+   earth_x_0+=     0.00000000090 * cos( 0.38965081024 +    58864.54391814630*t);
+   earth_x_0+=     0.00000000097 * cos( 0.27231106216 +    10557.59416082380*t);
+   earth_x_0+=     0.00000000101 * cos( 3.42372783781 +    19004.64794940840*t);
+   earth_x_0+=     0.00000000098 * cos( 3.23381285207 +     6126.67512948900*t);
+   earth_x_0+=     0.00000000098 * cos( 3.41472072058 +     6439.47657049380*t);
+   earth_x_0+=     0.00000000086 * cos( 1.34149823513 +    23123.74586080659*t);
+   earth_x_0+=     0.00000000118 * cos( 3.73846259947 +    13951.71327493390*t);
+   earth_x_0+=     0.00000000084 * cos( 3.13824782000 +     8858.31494432060*t);
+   earth_x_0+=     0.00000000117 * cos( 4.10879972365 +    89569.99011954499*t);
+   earth_x_0+=     0.00000000101 * cos( 0.23009068890 +    18429.74290609900*t);
+   earth_x_0+=     0.00000000088 * cos( 3.76745959676 +    17053.96910625319*t);
+   earth_x_0+=     0.00000000113 * cos( 2.05049922518 +    12985.63634385800*t);
+   earth_x_0+=     0.00000000082 * cos( 0.19375634761 +     5483.25472482600*t);
+   earth_x_0+=     0.00000000084 * cos( 4.23745405632 +     6379.05507720920*t);
+   earth_x_0+=     0.00000000084 * cos( 2.41107951633 +     6187.09662277360*t);
+   earth_x_0+=     0.00000000097 * cos( 4.02159784443 +    10239.58386601080*t);
+   earth_x_0+=     0.00000000105 * cos( 2.77478785571 +     6434.97313107660*t);
+   earth_x_0+=     0.00000000105 * cos( 3.87374571694 +     6131.17856890620*t);
+   earth_x_0+=     0.00000000114 * cos( 4.07080498725 +     8258.56839584740*t);
+   earth_x_0+=     0.00000000082 * cos( 3.82704107824 +     3958.12644117580*t);
+   earth_x_0+=     0.00000000093 * cos( 1.35189028328 +    18216.44381066100*t);
+   earth_x_0+=     0.00000000102 * cos( 2.89783593458 +     5767.61197889840*t);
+   earth_x_0+=     0.00000000102 * cos( 3.75069763806 +     6798.53972108440*t);
+   earth_x_0+=     0.00000000110 * cos( 2.19234510205 +    56600.27928952220*t);
+   earth_x_0+=     0.00000000112 * cos( 5.48905138227 +    12188.77809206700*t);
+   earth_x_0+=     0.00000000080 * cos( 3.59070813227 +    24080.03501677719*t);
+   earth_x_0+=     0.00000000088 * cos( 3.00398297222 +     1573.82042409880*t);
+   earth_x_0+=     0.00000000092 * cos( 4.36461689094 +   101426.20877096949*t);
+   earth_x_0+=     0.00000000087 * cos( 2.32187773582 +     1580.86666079680*t);
+   earth_x_0+=     0.00000000110 * cos( 0.37009059115 +    78423.70451667879*t);
+   earth_x_0+=     0.00000000099 * cos( 2.40898854628 +    33794.54372352860*t);
+   earth_x_0+=     0.00000000080 * cos( 3.27906531246 +     6485.32924516550*t);
+   earth_x_0+=     0.00000000080 * cos( 3.36946826019 +     6080.82245481730*t);
+   earth_x_0+=     0.00000000089 * cos( 5.20240793768 +    17363.24742890899*t);
+   earth_x_0+=     0.00000000079 * cos( 4.53817935583 +    36109.38220466459*t);
+   earth_x_0+=     0.00000000088 * cos( 0.87070535180 +    16943.76278503380*t);
+   earth_x_0+=     0.00000000088 * cos( 3.24901264480 +    12537.70251251500*t);
+   earth_x_0+=     0.00000000087 * cos( 1.75120977431 +    28237.23345938940*t);
+   earth_x_0+=     0.00000000080 * cos( 3.62211320521 +     3956.50801601940*t);
+   earth_x_0+=     0.00000000083 * cos( 0.01677699839 +    21424.46664430340*t);
+   earth_x_0+=     0.00000000099 * cos( 4.35778148149 +     5017.50837136500*t);
+   earth_x_0+=     0.00000000079 * cos( 4.50777133706 +     5621.84292321040*t);
+   earth_x_0+=     0.00000000096 * cos( 1.98497858569 +     5888.44996493220*t);
+   earth_x_0+=     0.00000000096 * cos( 4.66355498696 +     6677.70173505060*t);
+   earth_x_0+=     0.00000000074 * cos( 5.98124271909 +     6304.41149045840*t);
+   earth_x_0+=     0.00000000074 * cos( 0.66729085356 +     6261.74020952440*t);
+   earth_x_0+=     0.00000000073 * cos( 2.62129339640 +     6695.44694686580*t);
+   earth_x_0+=     0.00000000073 * cos( 4.02724017625 +     5870.70475311700*t);
+   earth_x_0+=     0.00000000085 * cos( 3.65241534932 +    12850.01101884800*t);
+   earth_x_0+=     0.00000000076 * cos( 1.33659857866 +    40796.33892267419*t);
+   earth_x_0+=     0.00000000097 * cos( 0.92251430623 +    24383.07910844140*t);
+   earth_x_0+=     0.00000000072 * cos( 0.32356695638 +     6312.89728814020*t);
+   earth_x_0+=     0.00000000072 * cos( 0.04178130909 +     6253.25441184259*t);
+   earth_x_0+=     0.00000000095 * cos( 0.62289841825 +    11080.17157891760*t);
+   earth_x_0+=     0.00000000100 * cos( 5.37794668456 +     5316.10497255580*t);
+   earth_x_0+=     0.00000000085 * cos( 2.90983978106 +    12282.29238111760*t);
+   earth_x_0+=     0.00000000071 * cos( 6.13573135109 +    41194.48792608240*t);
+   earth_x_0+=     0.00000000072 * cos( 4.76869158882 +    15141.39079431200*t);
+   earth_x_0+=     0.00000000075 * cos( 1.96590126974 +    24422.37035140730*t);
+   earth_x_0+=     0.00000000073 * cos( 1.29002951973 +    45892.73043315699*t);
+   earth_x_0+=     0.00000000076 * cos( 1.98229485596 +    63658.87775083760*t);
+   earth_x_0+=     0.00000000087 * cos( 3.63418769958 +    18106.23748944160*t);
+   earth_x_0+=     0.00000000070 * cos( 2.01720588184 +     6442.01936777460*t);
+   earth_x_0+=     0.00000000070 * cos( 4.63132769081 +     6124.13233220820*t);
+   earth_x_0+=     0.00000000097 * cos( 4.76857001908 +    92747.68916682259*t);
+   earth_x_0+=     0.00000000078 * cos( 4.94336223185 +    23550.34405168259*t);
+   earth_x_0+=     0.00000000076 * cos( 3.43423122344 +     5237.92101380380*t);
+   earth_x_0+=     0.00000000078 * cos( 4.98905027033 +     9924.81042151060*t);
+   earth_x_0+=     0.00000000082 * cos( 5.48703292615 +     6774.73914245020*t);
+   earth_x_0+=     0.00000000067 * cos( 0.17045052997 +     7082.89697515680*t);
+   earth_x_0+=     0.00000000067 * cos( 4.81312573338 +     6171.88920770380*t);
+   earth_x_0+=     0.00000000067 * cos( 1.83540783927 +     6394.26249227900*t);
+   earth_x_0+=     0.00000000080 * cos( 5.36760004078 +    22380.75580027400*t);
+   earth_x_0+=     0.00000000070 * cos( 2.04678665255 +     6175.41232605280*t);
+   earth_x_0+=     0.00000000070 * cos( 4.60174692010 +     6390.73937393000*t);
+   earth_x_0+=     0.00000000091 * cos( 2.47526506655 +     6652.77566593180*t);
+   earth_x_0+=     0.00000000091 * cos( 4.17326850610 +     5913.37603405100*t);
+   earth_x_0+=     0.00000000067 * cos( 0.13241010748 +     6300.88837210940*t);
+   earth_x_0+=     0.00000000067 * cos( 0.23293815799 +     6265.26332787340*t);
+   earth_x_0+=     0.00000000067 * cos( 1.43864813435 +     5959.57043333400*t);
+   earth_x_0+=     0.00000000080 * cos( 0.77942675353 +    23646.32327890039*t);
+   earth_x_0+=     0.00000000068 * cos( 3.64312434743 +     4897.18057365520*t);
+   earth_x_0+=     0.00000000069 * cos( 2.14076204045 +     6944.30877677240*t);
+   earth_x_0+=     0.00000000089 * cos( 1.97945887866 +    10660.68693504240*t);
+   earth_x_0+=     0.00000000073 * cos( 4.53414425046 +     6195.76764545159*t);
+   earth_x_0+=     0.00000000073 * cos( 2.11438932219 +     6370.38405453120*t);
+   earth_x_0+=     0.00000000066 * cos( 4.95039391556 +      323.50541665740*t);
+   earth_x_0+=     0.00000000066 * cos( 3.95274142218 +     6251.10354417460*t);
+   earth_x_0+=     0.00000000066 * cos( 2.69579215047 +     6315.04815580820*t);
+   earth_x_0+=     0.00000000063 * cos( 4.27995423229 +    11823.16163945020*t);
+   earth_x_0+=     0.00000000064 * cos( 0.09094709813 +    28774.03797148480*t);
+   earth_x_0+=     0.00000000086 * cos( 4.32355291567 +    12341.80690428090*t);
+   earth_x_0+=     0.00000000081 * cos( 0.47769423835 +    31441.67756975680*t);
+   earth_x_0+=     0.00000000071 * cos( 5.70281660725 +    16207.88627150200*t);
+   earth_x_0+=     0.00000000061 * cos( 0.43078256287 +    35050.00027447539*t);
+   earth_x_0+=     0.00000000066 * cos( 2.78136130877 +    17576.54652434700*t);
+   earth_x_0+=     0.00000000072 * cos( 0.54355547965 +    18100.00325845000*t);
+   earth_x_0+=     0.00000000072 * cos( 1.59783293112 +   245707.46610434429*t);
+   earth_x_0+=     0.00000000064 * cos( 1.93356156446 +    23006.42599258639*t);
+   earth_x_0+=     0.00000000068 * cos( 2.90815872826 +    28759.81087748319*t);
+   earth_x_0+=     0.00000000062 * cos( 0.67966053291 +     3855.03366695720*t);
+   earth_x_0+=     0.00000000073 * cos( 6.18334008725 +    23440.13773046320*t);
+   earth_x_0+=     0.00000000060 * cos( 3.22950754418 +    13088.72911807660*t);
+   earth_x_0+=     0.00000000078 * cos( 4.56094059497 +    12564.66722727450*t);
+   earth_x_0+=     0.00000000060 * cos( 4.47214922427 +     7548.64332861780*t);
+   earth_x_0+=     0.00000000075 * cos( 4.42126467141 +    18624.88275427230*t);
+   earth_x_0+=     0.00000000069 * cos( 5.82409328726 +   167993.69463622378*t);
+   earth_x_0+=     0.00000000059 * cos( 0.42392189557 +    41990.78593289880*t);
+   earth_x_0+=     0.00000000081 * cos( 1.27058672695 +     7250.04672742700*t);
+   earth_x_0+=     0.00000000059 * cos( 3.96273011871 +     6289.14176562120*t);
+   earth_x_0+=     0.00000000059 * cos( 2.68580345394 +     6277.00993436160*t);
+   earth_x_0+=     0.00000000063 * cos( 4.91003104939 +     5436.99301524020*t);
+   earth_x_0+=     0.00000000059 * cos( 2.73782977803 +     8584.66166590080*t);
+   earth_x_0+=     0.00000000074 * cos( 2.29556008890 +    28313.28880466100*t);
+   earth_x_0+=     0.00000000057 * cos( 0.01221429824 +    10027.90319572920*t);
+   earth_x_0+=     0.00000000057 * cos( 4.87299369650 +      792.77488846740*t);
+   earth_x_0+=     0.00000000060 * cos( 2.20527423485 +      647.01083331480*t);
+   earth_x_0+=     0.00000000056 * cos( 4.47895058873 +    15265.88651930040*t);
+   earth_x_0+=     0.00000000072 * cos( 4.16398251563 +     4598.02372748980*t);
+   earth_x_0+=     0.00000000063 * cos( 1.27506758028 +        5.52292430740*t);
+   earth_x_0+=     0.00000000068 * cos( 2.71579728238 +     3166.41643773160*t);
+   earth_x_0+=     0.00000000074 * cos( 3.23404191447 +    82534.40362061159*t);
+   earth_x_0+=     0.00000000056 * cos( 1.96997381249 +     6339.87911216120*t);
+   earth_x_0+=     0.00000000056 * cos( 4.67855976016 +     6226.27258782160*t);
+   earth_x_0+=     0.00000000069 * cos( 1.89738325082 +    23536.11695768099*t);
+   earth_x_0+=     0.00000000056 * cos( 2.52728721608 +    12303.06777661000*t);
+   earth_x_0+=     0.00000000055 * cos( 3.89144240866 +    22345.26037610820*t);
+   earth_x_0+=     0.00000000057 * cos( 3.55801157684 +     6702.00024889200*t);
+   earth_x_0+=     0.00000000057 * cos( 3.09052199581 +     5864.15145109080*t);
+   earth_x_0+=     0.00000000064 * cos( 5.42433316705 +    14169.79356223040*t);
+   earth_x_0+=     0.00000000075 * cos( 5.93374782872 +    32217.20018108080*t);
+   earth_x_0+=     0.00000000059 * cos( 2.17826568386 +    18003.14471522300*t);
+   earth_x_0+=     0.00000000059 * cos( 3.38504946794 +    67589.08744705719*t);
+   earth_x_0+=     0.00000000076 * cos( 2.34884503279 +    23226.83863502520*t);
+   earth_x_0+=     0.00000000073 * cos( 6.00950734574 +    30348.88377276700*t);
+   earth_x_0+=     0.00000000070 * cos( 1.56579326622 +    12299.54465826100*t);
+   earth_x_0+=     0.00000000060 * cos( 0.37221954582 +    12171.52581492360*t);
+   earth_x_0+=     0.00000000058 * cos( 3.15618940951 +     6716.78758786820*t);
+   earth_x_0+=     0.00000000074 * cos( 5.91246929008 +     3684.37804999020*t);
+   earth_x_0+=     0.00000000067 * cos( 2.65177830926 +    22743.40937951640*t);
+   earth_x_0+=     0.00000000062 * cos( 2.47017142882 +    16193.65917750039*t);
+   earth_x_0+=     0.00000000059 * cos( 2.60023110979 +       15.25247118500*t);
+   earth_x_0+=     0.00000000059 * cos( 1.05069387628 +    12164.47957822560*t);
+   earth_x_0+=     0.00000000056 * cos( 4.02457745520 +     3178.14579056760*t);
+   earth_x_0+=     0.00000000063 * cos( 3.66056370325 +     5219.76176653910*t);
+   earth_x_0+=     0.00000000063 * cos( 0.58993036842 +     5227.62607306530*t);
+   earth_x_0+=     0.00000000071 * cos( 3.61897462493 +    60530.48898574180*t);
+   earth_x_0+=     0.00000000072 * cos( 1.05408545034 +     4583.79663348820*t);
+   earth_x_0+=     0.00000000052 * cos( 5.67741459488 +     3627.50747751760*t);
+   earth_x_0+=     0.00000000052 * cos( 1.77495138030 +     6155.12431695680*t);
+   earth_x_0+=     0.00000000052 * cos( 4.87358219235 +     6411.02738302600*t);
+   earth_x_0+=     0.00000000073 * cos( 4.40543588100 +    80181.53746683979*t);
+   earth_x_0+=     0.00000000058 * cos( 4.29308512080 +    30775.48196364300*t);
+   earth_x_0+=     0.00000000068 * cos( 1.86654756299 +    35371.88726597640*t);
+   earth_x_0+=     0.00000000062 * cos( 4.46845874247 +     6908.74604230380*t);
+   earth_x_0+=     0.00000000062 * cos( 2.18007483018 +     5657.40565767900*t);
+   earth_x_0+=     0.00000000050 * cos( 1.53898120800 +     2575.23909432920*t);
+   earth_x_0+=     0.00000000052 * cos( 1.82842431930 +     5244.04923920100*t);
+   earth_x_0+=     0.00000000050 * cos( 4.69455799921 +     5881.40372823420*t);
+   earth_x_0+=     0.00000000050 * cos( 1.95397557344 +     6684.74797174860*t);
+   earth_x_0+=     0.00000000053 * cos( 5.33677558065 +    12146.66705610760*t);
+   earth_x_0+=     0.00000000068 * cos( 6.19033473088 +    96677.89886304220*t);
+   earth_x_0+=     0.00000000064 * cos( 5.16841754323 +    24279.10701821359*t);
+   earth_x_0+=     0.00000000050 * cos( 3.21675900656 +    43738.80234596580*t);
+   earth_x_0+=     0.00000000050 * cos( 1.95435790721 +    23141.55838292460*t);
+   earth_x_0+=     0.00000000059 * cos( 4.28456134455 +     3981.49003408200*t);
+   earth_x_0+=     0.00000000058 * cos( 2.65359377943 +    40077.61957352000*t);
+   earth_x_0+=     0.00000000054 * cos( 3.56564135998 +    15664.03552270859*t);
+   earth_x_0+=     0.00000000061 * cos( 2.88962487651 +    40398.18991926600*t);
+   earth_x_0+=     0.00000000051 * cos( 3.00540767523 +     7668.97112632760*t);
+   earth_x_0+=     0.00000000048 * cos( 4.80656273565 +     3646.35037735440*t);
+   earth_x_0+=     0.00000000048 * cos( 2.80573655186 +      170.67287061920*t);
+   earth_x_0+=     0.00000000065 * cos( 6.08902325889 +    47162.51635463520*t);
+   earth_x_0+=     0.00000000056 * cos( 4.19093499866 +      799.82112516540*t);
+   earth_x_0+=     0.00000000049 * cos( 3.72443814803 +    96562.99901809469*t);
+   earth_x_0+=     0.00000000050 * cos( 5.17709552474 +    17583.66007134780*t);
+   earth_x_0+=     0.00000000049 * cos( 5.16566731300 +     5490.30096152400*t);
+   earth_x_0+=     0.00000000058 * cos( 5.40778008235 +     6076.37511269480*t);
+   earth_x_0+=     0.00000000058 * cos( 1.24075349030 +     6489.77658728800*t);
+   earth_x_0+=     0.00000000059 * cos( 1.72439508836 +    44034.12758953940*t);
+   earth_x_0+=     0.00000000045 * cos( 1.62661962717 +     4480.70385926960*t);
+   earth_x_0+=     0.00000000061 * cos( 5.10264428859 +      491.66329245880*t);
+   earth_x_0+=     0.00000000044 * cos( 1.10033696885 +     2301.58581590939*t);
+   earth_x_0+=     0.00000000054 * cos( 3.85123315789 +     6559.82162185580*t);
+   earth_x_0+=     0.00000000054 * cos( 2.79730041476 +     6006.33007812700*t);
+   earth_x_0+=     0.00000000052 * cos( 5.42245980619 +   244287.60000722769*t);
+   earth_x_0+=     0.00000000048 * cos( 1.42585982742 +    36147.40987730040*t);
+   earth_x_0+=     0.00000000044 * cos( 0.81389810430 +     5113.48759858280*t);
+   earth_x_0+=     0.00000000045 * cos( 4.22954077900 +      533.21408344360*t);
+   earth_x_0+=     0.00000000046 * cos( 2.27703250571 +     4384.72463205180*t);
+   earth_x_0+=     0.00000000048 * cos( 6.20173890368 +    35579.69123957000*t);
+   earth_x_0+=     0.00000000051 * cos( 0.20646603354 +    18326.65013188039*t);
+   earth_x_0+=     0.00000000043 * cos( 3.18939503640 +    27832.03821928320*t);
+   earth_x_0+=     0.00000000046 * cos( 4.72419541485 +    16703.06213349900*t);
+   earth_x_0+=     0.00000000052 * cos( 4.17179726389 +    11919.14086666800*t);
+   earth_x_0+=     0.00000000052 * cos( 0.71926561105 +    46386.99374331120*t);
+   earth_x_0+=     0.00000000047 * cos( 0.82248881139 +     5209.46682580060*t);
+   earth_x_0+=     0.00000000046 * cos( 4.04752097938 +    12555.10599971890*t);
+   earth_x_0+=     0.00000000052 * cos( 0.43326734132 +     6574.04871585740*t);
+   earth_x_0+=     0.00000000052 * cos( 6.21526623133 +     5992.10298412540*t);
+   earth_x_0+=     0.00000000044 * cos( 5.37918791865 +    69941.95360082899*t);
+   earth_x_0+=     0.00000000058 * cos( 0.36325356715 +    74673.98422817240*t);
+   earth_x_0+=     0.00000000056 * cos( 1.50089021543 +    27707.54249429480*t);
+   earth_x_0+=     0.00000000045 * cos( 6.22182195617 +    46360.69542351140*t);
+   earth_x_0+=     0.00000000043 * cos( 3.11608101277 +     8401.83971036980*t);
+   earth_x_0+=     0.00000000047 * cos( 2.39687793686 +   143571.32428481648*t);
+   earth_x_0+=     0.00000000043 * cos( 1.13301289185 +    30639.85663863300*t);
+   earth_x_0+=     0.00000000049 * cos( 1.97844011728 +    22030.21295466960*t);
+   earth_x_0+=     0.00000000045 * cos( 4.79114307475 +    71519.29714327680*t);
+   earth_x_0+=     0.00000000051 * cos( 4.52703512787 +     7761.94242405580*t);
+   earth_x_0+=     0.00000000041 * cos( 0.43156447693 +     3735.23831175900*t);
+   earth_x_0+=     0.00000000041 * cos( 1.66415755313 +    44136.95134937400*t);
+   earth_x_0+=     0.00000000047 * cos( 4.45924253669 +     4708.23004870920*t);
+   earth_x_0+=     0.00000000056 * cos( 0.65498687965 +    28663.83165026540*t);
+   earth_x_0+=     0.00000000055 * cos( 3.26682599355 +    39302.09696219600*t);
+   earth_x_0+=     0.00000000049 * cos( 5.40929899528 +    27511.46787353720*t);
+   earth_x_0+=     0.00000000047 * cos( 0.89791015009 +   316428.22867391503*t);
+   earth_x_0+=     0.00000000040 * cos( 0.05000438602 +    18208.34994259200*t);
+   earth_x_0+=     0.00000000056 * cos( 2.85120720656 +    83659.27687239899*t);
+   earth_x_0+=     0.00000000053 * cos( 3.02524199997 +    38650.17350619900*t);
+   earth_x_0+=     0.00000000048 * cos( 2.50261741205 +    12662.13092720060*t);
+   earth_x_0+=     0.00000000052 * cos( 4.41165168727 +    12571.67462429020*t);
+   earth_x_0+=     0.00000000043 * cos( 4.69793306978 +     7322.10246078180*t);
+   earth_x_0+=     0.00000000041 * cos( 3.56541409677 +      664.75604513000*t);
+   earth_x_0+=     0.00000000044 * cos( 5.20087497624 +     6140.89722295520*t);
+   earth_x_0+=     0.00000000044 * cos( 1.44765859641 +     6425.25447702760*t);
+   earth_x_0+=     0.00000000039 * cos( 0.38207420659 +    22594.05489571199*t);
+   earth_x_0+=     0.00000000039 * cos( 4.22033276864 +    27278.46881644080*t);
+   earth_x_0+=     0.00000000045 * cos( 4.12311427457 +    44007.82926973960*t);
+   earth_x_0+=     0.00000000054 * cos( 4.69164296803 +     6006.04085625000*t);
+   earth_x_0+=     0.00000000054 * cos( 1.95689060462 +     6560.11084373280*t);
+   earth_x_0+=     0.00000000045 * cos( 5.54498349441 +   142861.39123625819*t);
+   earth_x_0+=     0.00000000047 * cos( 1.12159858858 +    24499.51966065240*t);
+   earth_x_0+=     0.00000000043 * cos( 1.36615498859 +     7026.06591052400*t);
+   earth_x_0+=     0.00000000040 * cos( 0.90738693747 +    26880.31981303260*t);
+   earth_x_0+=     0.00000000047 * cos( 4.50830744099 +    24485.29256665079*t);
+   earth_x_0+=     0.00000000043 * cos( 5.81439105801 +     6359.34192126700*t);
+   earth_x_0+=     0.00000000042 * cos( 4.70083045322 +     6231.79551212899*t);
+   earth_x_0+=     0.00000000042 * cos( 1.94770311943 +     6334.35618785381*t);
+   earth_x_0+=     0.00000000037 * cos( 5.39982800110 +    23020.65308658799*t);
+   earth_x_0+=     0.00000000052 * cos( 3.53311914292 +    11904.91877320180*t);
+   earth_x_0+=     0.00000000040 * cos( 1.50947925293 +     7075.85073845880*t);
+   earth_x_0+=     0.00000000049 * cos( 3.73016699917 +    32243.49850088060*t);
+   earth_x_0+=     0.00000000043 * cos( 4.45265815492 +    40000.04091585779*t);
+   earth_x_0+=     0.00000000040 * cos( 0.07156369559 +     7122.04513774180*t);
+   earth_x_0+=     0.00000000040 * cos( 0.29378456988 +     5444.10656224100*t);
+   earth_x_0+=     0.00000000038 * cos( 2.24527247949 +    44933.24935619040*t);
+   earth_x_0+=     0.00000000037 * cos( 0.09116813653 +     6850.79448772180*t);
+   earth_x_0+=     0.00000000037 * cos( 0.27418012894 +     5715.35721226100*t);
+   earth_x_0+=     0.00000000039 * cos( 2.88920508037 +       20.35531939880*t);
+   earth_x_0+=     0.00000000051 * cos( 5.82785362503 +    84111.74716305939*t);
+   earth_x_0+=     0.00000000037 * cos( 1.90653942815 +    16723.35014259500*t);
+   earth_x_0+=     0.00000000042 * cos( 3.79031051907 +     6328.21706962800*t);
+   earth_x_0+=     0.00000000042 * cos( 2.85822305358 +     6237.93463035480*t);
+   earth_x_0+=     0.00000000036 * cos( 2.35204592553 +    20760.42703319140*t);
+   earth_x_0+=     0.00000000049 * cos( 4.54040658078 +    11293.47067435560*t);
+   earth_x_0+=     0.00000000040 * cos( 3.64328287226 +     3215.13638064320*t);
+   earth_x_0+=     0.00000000050 * cos( 1.62788844152 +    84334.41776382479*t);
+   earth_x_0+=     0.00000000040 * cos( 0.78396796338 +    71980.63357473118*t);
+
+   double earth_x_1=0.0;
+   earth_x_1+=     0.00123403056 * cos( 0.00000000000 +        0.00000000000*t);
+   earth_x_1+=     0.00051500156 * cos( 6.00266267204 +    12566.15169998280*t);
+   earth_x_1+=     0.00001290726 * cos( 5.95943124583 +    18849.22754997420*t);
+   earth_x_1+=     0.00001068627 * cos( 2.01554176551 +     6283.07584999140*t);
+   earth_x_1+=     0.00000212689 * cos( 1.73380190491 +     6279.55273164240*t);
+   earth_x_1+=     0.00000212515 * cos( 4.91489371033 +     6286.59896834040*t);
+   earth_x_1+=     0.00000062260 * cos( 0.36239798178 +     4705.73230754360*t);
+   earth_x_1+=     0.00000059822 * cos( 3.81195369871 +     6256.77753019160*t);
+   earth_x_1+=     0.00000059514 * cos( 2.83634160150 +     6309.37416979120*t);
+   earth_x_1+=     0.00000048841 * cos( 5.21419389335 +      775.52261132400*t);
+   earth_x_1+=     0.00000042883 * cos( 0.43789776559 +     1059.38193018920*t);
+   earth_x_1+=     0.00000046286 * cos( 0.01839494103 +     7860.41939243920*t);
+   earth_x_1+=     0.00000035675 * cos( 1.45279327264 +     5884.92684658320*t);
+   earth_x_1+=     0.00000036061 * cos( 2.16002201071 +     5753.38488489680*t);
+   earth_x_1+=     0.00000035367 * cos( 4.47243820095 +     6812.76681508600*t);
+   earth_x_1+=     0.00000032137 * cos( 5.19589851893 +     6681.22485339960*t);
+   earth_x_1+=     0.00000028763 * cos( 5.91618989512 +    25132.30339996560*t);
+   earth_x_1+=     0.00000028447 * cos( 1.14976253807 +     6127.65545055720*t);
+   earth_x_1+=     0.00000027573 * cos( 5.50119104683 +     6438.49624942560*t);
+   earth_x_1+=     0.00000024815 * cos( 2.92204909812 +     5486.77784317500*t);
+   earth_x_1+=     0.00000020611 * cos( 3.71790880968 +     7079.37385680780*t);
+   earth_x_1+=     0.00000019565 * cos( 2.89351924469 +     5507.55323866740*t);
+   earth_x_1+=     0.00000018308 * cos( 1.46954314992 +    11790.62908865880*t);
+   earth_x_1+=     0.00000016471 * cos( 6.22682639292 +    11506.76976979360*t);
+   earth_x_1+=     0.00000016757 * cos( 3.81935015812 +     7058.59846131540*t);
+   earth_x_1+=     0.00000014555 * cos( 5.97554823531 +     6290.18939699220*t);
+   earth_x_1+=     0.00000014388 * cos( 0.68157599309 +     6275.96230299060*t);
+   earth_x_1+=     0.00000013921 * cos( 1.44156172409 +      796.29800681640*t);
+   earth_x_1+=     0.00000011845 * cos( 4.15246503623 +     4694.00295470760*t);
+   earth_x_1+=     0.00000012573 * cos( 0.30429978572 +        7.11354700080*t);
+   earth_x_1+=     0.00000010073 * cos( 3.28940171828 +     3738.76143010800*t);
+   earth_x_1+=     0.00000009513 * cos( 4.62657118532 +     6069.77675455340*t);
+   earth_x_1+=     0.00000010425 * cos( 4.26610810148 +     6282.09552892320*t);
+   earth_x_1+=     0.00000010425 * cos( 2.38242547117 +     6284.05617105960*t);
+   earth_x_1+=     0.00000009565 * cos( 1.89315044337 +     4136.91043351620*t);
+   earth_x_1+=     0.00000009044 * cos( 2.00977420682 +     6496.37494542940*t);
+   earth_x_1+=     0.00000008543 * cos( 0.05563879210 +     1194.44701022460*t);
+   earth_x_1+=     0.00000008230 * cos( 1.25204318279 +     1589.07289528380*t);
+   earth_x_1+=     0.00000006181 * cos( 3.36235656603 +     8827.39026987480*t);
+   earth_x_1+=     0.00000006252 * cos( 4.75560512125 +     8429.24126646660*t);
+   earth_x_1+=     0.00000005525 * cos( 5.36960259598 +     4933.20844033260*t);
+   earth_x_1+=     0.00000007312 * cos( 0.49039562572 +     4535.05943692440*t);
+   earth_x_1+=     0.00000007074 * cos( 4.39447155183 +     5088.62883976680*t);
+   earth_x_1+=     0.00000005168 * cos( 5.15700586490 +    11769.85369316640*t);
+   earth_x_1+=     0.00000004878 * cos( 0.63279675325 +     6040.34724601740*t);
+   earth_x_1+=     0.00000005339 * cos( 5.63168593753 +     3154.68708489560*t);
+   earth_x_1+=     0.00000004985 * cos( 4.62777616911 +    12569.67481833180*t);
+   earth_x_1+=     0.00000004902 * cos( 4.14599178687 +     3930.20969621960*t);
+   earth_x_1+=     0.00000004709 * cos( 3.96401482904 +     5331.35744374080*t);
+   earth_x_1+=     0.00000004601 * cos( 6.02376082863 +     6525.80445396540*t);
+   earth_x_1+=     0.00000004331 * cos( 1.32999081524 +     7632.94325965020*t);
+   earth_x_1+=     0.00000005607 * cos( 2.55800588736 +     5729.50644714900*t);
+   earth_x_1+=     0.00000003968 * cos( 2.72696169781 +     7234.79425624200*t);
+   earth_x_1+=     0.00000004394 * cos( 0.04709283876 +    12559.03815298200*t);
+   earth_x_1+=     0.00000005262 * cos( 6.19471325132 +     8031.09226305840*t);
+   earth_x_1+=     0.00000005100 * cos( 4.10862449537 +     6836.64525283380*t);
+   earth_x_1+=     0.00000005075 * cos( 2.30582917471 +     7477.52286021600*t);
+   earth_x_1+=     0.00000003287 * cos( 5.41013501104 +    10977.07880469900*t);
+   earth_x_1+=     0.00000003366 * cos( 0.28947921379 +    11371.70468975820*t);
+   earth_x_1+=     0.00000003341 * cos( 1.64988056630 +     5223.69391980220*t);
+   earth_x_1+=     0.00000003495 * cos( 3.73571049570 +     4164.31198961300*t);
+   earth_x_1+=     0.00000004294 * cos( 4.93755825416 +     1592.59601363280*t);
+   earth_x_1+=     0.00000003408 * cos( 2.40593164878 +     3128.38876509580*t);
+   earth_x_1+=     0.00000002997 * cos( 4.86533819655 +     1748.01641306700*t);
+   earth_x_1+=     0.00000002750 * cos( 5.24429623672 +     7342.45778018060*t);
+   earth_x_1+=     0.00000002669 * cos( 3.42836090012 +     2146.16541647520*t);
+   earth_x_1+=     0.00000003313 * cos( 1.92016526057 +      801.82093112380*t);
+   earth_x_1+=     0.00000002991 * cos( 2.26433580653 +     8635.94200376320*t);
+   earth_x_1+=     0.00000003390 * cos( 4.23836515895 +      155.42039943420*t);
+   earth_x_1+=     0.00000002384 * cos( 5.74836325186 +     7084.89678111520*t);
+   earth_x_1+=     0.00000002288 * cos( 4.51088946570 +    17260.15465469040*t);
+   earth_x_1+=     0.00000002405 * cos( 5.96195251565 +     4690.47983635860*t);
+   earth_x_1+=     0.00000002052 * cos( 1.42012676041 +      951.71840625060*t);
+   earth_x_1+=     0.00000002049 * cos( 3.53634954129 +     1990.74501704100*t);
+   earth_x_1+=     0.00000002778 * cos( 1.46151348750 +     5481.25491886760*t);
+   earth_x_1+=     0.00000002656 * cos( 5.71405746008 +       26.29831979980*t);
+   earth_x_1+=     0.00000002135 * cos( 1.89020345470 +      536.80451209540*t);
+   earth_x_1+=     0.00000002529 * cos( 0.00975556739 +     1349.86740965880*t);
+   earth_x_1+=     0.00000001877 * cos( 2.82846563483 +      553.56940284240*t);
+   earth_x_1+=     0.00000001858 * cos( 0.74657310664 +      398.14900340820*t);
+   earth_x_1+=     0.00000002013 * cos( 2.40172210263 +      522.57741809380*t);
+   earth_x_1+=     0.00000001904 * cos( 4.23263119723 +     9437.76293488700*t);
+   earth_x_1+=     0.00000001708 * cos( 0.29065876129 +    10988.80815753500*t);
+   earth_x_1+=     0.00000001853 * cos( 1.84012190435 +     5230.80746680300*t);
+   earth_x_1+=     0.00000001857 * cos( 2.34009177634 +     5216.58037280140*t);
+   earth_x_1+=     0.00000001558 * cos( 1.10024744060 +      426.59819087600*t);
+   earth_x_1+=     0.00000001543 * cos( 5.14512865944 +    10447.38783960440*t);
+   earth_x_1+=     0.00000001555 * cos( 3.67004205114 +    13095.84266507740*t);
+   earth_x_1+=     0.00000001779 * cos( 1.80817581905 +    12562.62858163380*t);
+   earth_x_1+=     0.00000001450 * cos( 0.99962470222 +     6262.72053059260*t);
+   earth_x_1+=     0.00000001450 * cos( 5.64890887043 +     6303.43116939020*t);
+   earth_x_1+=     0.00000001880 * cos( 1.74872544631 +    10973.55568635000*t);
+   earth_x_1+=     0.00000001372 * cos( 0.89279228106 +     7875.67186362420*t);
+   earth_x_1+=     0.00000001247 * cos( 1.86318425701 +     2118.76386037840*t);
+   earth_x_1+=     0.00000001231 * cos( 4.10277044176 +    18319.53658487960*t);
+   earth_x_1+=     0.00000001373 * cos( 4.83873686423 +    12721.57209941700*t);
+   earth_x_1+=     0.00000001140 * cos( 3.49255962548 +     1066.49547719000*t);
+   earth_x_1+=     0.00000001127 * cos( 6.04843564608 +     5643.17856367740*t);
+   earth_x_1+=     0.00000001136 * cos( 2.83389341769 +     2942.46342329160*t);
+   earth_x_1+=     0.00000001071 * cos( 3.77511422440 +      242.72860397400*t);
+   earth_x_1+=     0.00000001165 * cos( 5.90779831262 +    15720.83878487840*t);
+   earth_x_1+=     0.00000000974 * cos( 0.24120093775 +      639.89728631400*t);
+   earth_x_1+=     0.00000000973 * cos( 1.30526267473 +     4292.33083295040*t);
+   earth_x_1+=     0.00000000993 * cos( 3.20822379065 +    10575.40668294180*t);
+   earth_x_1+=     0.00000000964 * cos( 2.13521649390 +     2388.89402044920*t);
+   earth_x_1+=     0.00000000929 * cos( 1.06241249364 +    16496.36139620240*t);
+   earth_x_1+=     0.00000000939 * cos( 5.78268735385 +    20426.57109242200*t);
+   earth_x_1+=     0.00000000889 * cos( 0.06581239696 +     6288.59877429880*t);
+   earth_x_1+=     0.00000000889 * cos( 0.29953586851 +     6277.55292568400*t);
+   earth_x_1+=     0.00000001018 * cos( 4.23379398746 +    12539.85338018300*t);
+   earth_x_1+=     0.00000001042 * cos( 2.30738424939 +      529.69096509460*t);
+   earth_x_1+=     0.00000001033 * cos( 1.04793549544 +     5760.49843189760*t);
+   earth_x_1+=     0.00000001017 * cos( 5.59172649862 +     6805.65326808520*t);
+   earth_x_1+=     0.00000000801 * cos( 5.36423767778 +    14314.16811304980*t);
+   earth_x_1+=     0.00000000795 * cos( 4.97031461103 +     6279.78949257360*t);
+   earth_x_1+=     0.00000000795 * cos( 1.67821896162 +     6286.36220740920*t);
+   earth_x_1+=     0.00000000723 * cos( 3.87714868475 +     9623.68827669120*t);
+   earth_x_1+=     0.00000000800 * cos( 4.75445675725 +    12029.34718788740*t);
+   earth_x_1+=     0.00000000698 * cos( 0.23124375738 +     6148.01076995600*t);
+   earth_x_1+=     0.00000000675 * cos( 5.56136235734 +    12964.30070339100*t);
+   earth_x_1+=     0.00000000678 * cos( 0.13408273301 +     6418.14093002680*t);
+   earth_x_1+=     0.00000000671 * cos( 1.28940504615 +     6709.67404086740*t);
+   earth_x_1+=     0.00000000898 * cos( 2.26352977340 +    13367.97263110660*t);
+   earth_x_1+=     0.00000000716 * cos( 4.07157593048 +    16730.46368959580*t);
+   earth_x_1+=     0.00000000651 * cos( 5.31876664276 +     5856.47765911540*t);
+   earth_x_1+=     0.00000000744 * cos( 4.25820617142 +    12043.57428188900*t);
+   earth_x_1+=     0.00000000650 * cos( 3.96082074858 +    14712.31711645800*t);
+   earth_x_1+=     0.00000000643 * cos( 6.14575788122 +     2352.86615377180*t);
+   earth_x_1+=     0.00000000666 * cos( 6.18664331112 +     5746.27133789600*t);
+   earth_x_1+=     0.00000000588 * cos( 3.24442116878 +    13119.72110282519*t);
+   earth_x_1+=     0.00000000627 * cos( 1.85323853183 +     4686.88940770680*t);
+   earth_x_1+=     0.00000000584 * cos( 1.78420561597 +    13517.87010623340*t);
+   earth_x_1+=     0.00000000610 * cos( 0.48716946879 +     6819.88036208680*t);
+   earth_x_1+=     0.00000000666 * cos( 3.29946942274 +     3634.62102451840*t);
+   earth_x_1+=     0.00000000585 * cos( 0.14860338726 +    12036.46073488820*t);
+   earth_x_1+=     0.00000000620 * cos( 5.87390167695 +    31415.37924995700*t);
+   earth_x_1+=     0.00000000553 * cos( 2.36074148933 +     6306.95428773920*t);
+   earth_x_1+=     0.00000000553 * cos( 4.28779208332 +     6259.19741224360*t);
+   earth_x_1+=     0.00000000527 * cos( 4.65845256725 +    10177.25767953360*t);
+   earth_x_1+=     0.00000000523 * cos( 5.55991825992 +    12139.55350910680*t);
+   earth_x_1+=     0.00000000514 * cos( 2.30343678959 +     4701.11650170840*t);
+   earth_x_1+=     0.00000000495 * cos( 3.58067041581 +    13362.44970679920*t);
+   earth_x_1+=     0.00000000688 * cos( 0.41160312810 +    13916.01910964160*t);
+   earth_x_1+=     0.00000000481 * cos( 4.14197696514 +     6286.95718534940*t);
+   earth_x_1+=     0.00000000481 * cos( 2.50655660751 +     6279.19451463340*t);
+   earth_x_1+=     0.00000000465 * cos( 0.81901657744 +    14143.49524243060*t);
+   earth_x_1+=     0.00000000440 * cos( 2.91975126929 +     3894.18182954220*t);
+   earth_x_1+=     0.00000000431 * cos( 3.09968854047 +    11499.65622279280*t);
+   earth_x_1+=     0.00000000452 * cos( 5.55409870988 +    12573.26524698360*t);
+   earth_x_1+=     0.00000000499 * cos( 1.08746978287 +    12410.73130054860*t);
+   earth_x_1+=     0.00000000440 * cos( 2.14902528642 +    24072.92146977640*t);
+   earth_x_1+=     0.00000000435 * cos( 4.76945528790 +     9917.69687450980*t);
+   earth_x_1+=     0.00000000398 * cos( 5.77181168519 +     8273.82086703240*t);
+   earth_x_1+=     0.00000000379 * cos( 4.62464502982 +     6386.16862421000*t);
+   earth_x_1+=     0.00000000379 * cos( 2.02388854283 +     6179.98307577280*t);
+   earth_x_1+=     0.00000000377 * cos( 5.90355422208 +     7872.14874527520*t);
+   earth_x_1+=     0.00000000419 * cos( 5.62219999256 +     5327.47610838280*t);
+   earth_x_1+=     0.00000000383 * cos( 0.40420536338 +     6319.10371666880*t);
+   earth_x_1+=     0.00000000383 * cos( 6.24432820927 +     6247.04798331399*t);
+   earth_x_1+=     0.00000000486 * cos( 3.05618432546 +    12592.45001978260*t);
+   earth_x_1+=     0.00000000457 * cos( 2.45652916518 +     6489.26139842860*t);
+   earth_x_1+=     0.00000000457 * cos( 4.19200440747 +     6076.89030155420*t);
+   earth_x_1+=     0.00000000512 * cos( 1.88365818219 +    11015.10647733480*t);
+   earth_x_1+=     0.00000000344 * cos( 0.28667648401 +    18451.07854656599*t);
+   earth_x_1+=     0.00000000340 * cos( 2.75068444234 +    12323.42309600880*t);
+   earth_x_1+=     0.00000000454 * cos( 0.72878821246 +     2787.04302385740*t);
+   earth_x_1+=     0.00000000364 * cos( 2.25070711701 +     6321.10352262720*t);
+   earth_x_1+=     0.00000000364 * cos( 4.39782645564 +     6245.04817735560*t);
+   earth_x_1+=     0.00000000342 * cos( 6.10019640486 +     6303.85124548380*t);
+   earth_x_1+=     0.00000000342 * cos( 0.54833716779 +     6262.30045449900*t);
+   earth_x_1+=     0.00000000324 * cos( 2.66245548691 +    77713.77146812050*t);
+   earth_x_1+=     0.00000000316 * cos( 0.14807870151 +    11926.25441366880*t);
+   earth_x_1+=     0.00000000318 * cos( 1.32350341546 +     6311.52503745920*t);
+   earth_x_1+=     0.00000000318 * cos( 5.32503015719 +     6254.62666252360*t);
+   earth_x_1+=     0.00000000304 * cos( 2.00970548291 +    12779.45079542080*t);
+   earth_x_1+=     0.00000000360 * cos( 4.10116091510 +     1052.26838318840*t);
+   earth_x_1+=     0.00000000343 * cos( 2.93550529951 +     6267.82337880639*t);
+   earth_x_1+=     0.00000000343 * cos( 3.71302827314 +     6298.32832117640*t);
+   earth_x_1+=     0.00000000398 * cos( 0.12276624389 +     1551.04522264800*t);
+   earth_x_1+=     0.00000000300 * cos( 1.77610616438 +    90279.92316810328*t);
+   earth_x_1+=     0.00000000367 * cos( 2.75102591636 +    10021.83728009940*t);
+   earth_x_1+=     0.00000000297 * cos( 0.51546903036 +     5429.87946823940*t);
+   earth_x_1+=     0.00000000337 * cos( 2.71244267795 +    17782.73207278420*t);
+   earth_x_1+=     0.00000000295 * cos( 4.05850212848 +     6702.56049386660*t);
+   earth_x_1+=     0.00000000331 * cos( 0.94287498279 +     7238.67559160000*t);
+   earth_x_1+=     0.00000000321 * cos( 2.23719619346 +    17796.95916678580*t);
+   earth_x_1+=     0.00000000378 * cos( 2.47077025009 +      213.29909543800*t);
+   earth_x_1+=     0.00000000268 * cos( 4.36573038796 +     2544.31441988340*t);
+   earth_x_1+=     0.00000000290 * cos( 5.31163321847 +      220.41264243880*t);
+   earth_x_1+=     0.00000000349 * cos( 0.95962542267 +    17789.84561978500*t);
+   earth_x_1+=     0.00000000275 * cos( 3.67150875811 +    14945.31617355440*t);
+   earth_x_1+=     0.00000000253 * cos( 2.22918080251 +    18073.70493865020*t);
+   earth_x_1+=     0.00000000257 * cos( 1.04749733742 +      955.59974160860*t);
+   earth_x_1+=     0.00000000234 * cos( 4.07399715957 +     1596.18644228460*t);
+   earth_x_1+=     0.00000000282 * cos( 4.50260494490 +    13341.67431130680*t);
+   earth_x_1+=     0.00000000280 * cos( 5.71871851487 +      206.18554843720*t);
+   earth_x_1+=     0.00000000269 * cos( 6.10957634874 +     9779.10867612540*t);
+   earth_x_1+=     0.00000000306 * cos( 3.20523259534 +     5863.59120611620*t);
+   earth_x_1+=     0.00000000216 * cos( 5.60561646872 +     3185.19202726560*t);
+   earth_x_1+=     0.00000000280 * cos( 1.91082976057 +     5642.19824260920*t);
+   earth_x_1+=     0.00000000233 * cos( 3.01339621844 +    10984.19235169980*t);
+   earth_x_1+=     0.00000000263 * cos( 2.16888500036 +    12168.00269657460*t);
+   earth_x_1+=     0.00000000226 * cos( 4.44534966925 +     3496.03282613400*t);
+   earth_x_1+=     0.00000000260 * cos( 3.16564185469 +    10969.96525769820*t);
+   earth_x_1+=     0.00000000227 * cos( 4.95666866168 +     5849.36411211460*t);
+   earth_x_1+=     0.00000000211 * cos( 1.22111379347 +    19651.04848109800*t);
+   earth_x_1+=     0.00000000250 * cos( 4.61595902292 +     6923.95345737360*t);
+   earth_x_1+=     0.00000000194 * cos( 2.37630770389 +      640.87760738220*t);
+   earth_x_1+=     0.00000000203 * cos( 4.39496828470 +     6062.66320755260*t);
+   earth_x_1+=     0.00000000203 * cos( 2.25356528795 +     6503.48849243020*t);
+   earth_x_1+=     0.00000000184 * cos( 0.13859812924 +    18209.33026366019*t);
+   earth_x_1+=     0.00000000184 * cos( 4.20955149479 +    12566.21901028560*t);
+   earth_x_1+=     0.00000000180 * cos( 5.76799248594 +      853.19638175200*t);
+   earth_x_1+=     0.00000000199 * cos( 2.84103239065 +        3.52311834900*t);
+   earth_x_1+=     0.00000000187 * cos( 4.59167688742 +     6276.02961329339*t);
+   earth_x_1+=     0.00000000187 * cos( 2.05685668523 +     6290.12208668940*t);
+   earth_x_1+=     0.00000000192 * cos( 2.99578290747 +     4732.03062734340*t);
+   earth_x_1+=     0.00000000184 * cos( 4.65365446464 +    12566.08438968000*t);
+   earth_x_1+=     0.00000000230 * cos( 2.23774195109 +    12567.13202105100*t);
+   earth_x_1+=     0.00000000165 * cos( 1.34307380910 +      110.20632121940*t);
+   earth_x_1+=     0.00000000166 * cos( 1.19477871058 +      433.71173787680*t);
+   earth_x_1+=     0.00000000184 * cos( 3.38049914956 +     6037.24420376200*t);
+   earth_x_1+=     0.00000000156 * cos( 2.34457936544 +     2648.45482547300*t);
+   earth_x_1+=     0.00000000175 * cos( 3.26771507312 +     6528.90749622080*t);
+   earth_x_1+=     0.00000000154 * cos( 5.67725689556 +     6453.74872061060*t);
+   earth_x_1+=     0.00000000154 * cos( 0.97127667708 +     6112.40297937220*t);
+   earth_x_1+=     0.00000000166 * cos( 4.51313560848 +    12352.85260454480*t);
+   earth_x_1+=     0.00000000146 * cos( 4.05664802756 +      149.56319713460*t);
+   earth_x_1+=     0.00000000172 * cos( 0.88734623011 +     1577.34354244780*t);
+   earth_x_1+=     0.00000000144 * cos( 3.76179690789 +    18635.92845453620*t);
+   earth_x_1+=     0.00000000146 * cos( 1.99164761433 +    13760.59871020740*t);
+   earth_x_1+=     0.00000000136 * cos( 0.75192359085 +    11712.95531823080*t);
+   earth_x_1+=     0.00000000142 * cos( 1.34533369944 +     4157.19844261220*t);
+   earth_x_1+=     0.00000000133 * cos( 0.01824912792 +    18139.29450141590*t);
+   earth_x_1+=     0.00000000124 * cos( 0.99880484927 +    17256.63153634140*t);
+   earth_x_1+=     0.00000000140 * cos( 1.25451695814 +     9380.95967271720*t);
+   earth_x_1+=     0.00000000123 * cos( 0.26259126301 +     6172.86952877200*t);
+   earth_x_1+=     0.00000000123 * cos( 0.10275700246 +     6393.28217121080*t);
+   earth_x_1+=     0.00000000121 * cos( 5.40082218265 +       20.77539549240*t);
+   earth_x_1+=     0.00000000117 * cos( 4.47711713051 +    18852.75066832320*t);
+   earth_x_1+=     0.00000000114 * cos( 2.33765655466 +    65147.61976813770*t);
+   earth_x_1+=     0.00000000120 * cos( 2.29184279740 +    17253.04110768959*t);
+   earth_x_1+=     0.00000000132 * cos( 4.94196827917 +    18875.52586977400*t);
+   earth_x_1+=     0.00000000123 * cos( 4.42839393502 +    12565.17137891460*t);
+   earth_x_1+=     0.00000000129 * cos( 2.56187150433 +    11513.88331679440*t);
+   earth_x_1+=     0.00000000121 * cos( 2.27781803851 +    10419.98628350760*t);
+   earth_x_1+=     0.00000000112 * cos( 0.58654969541 +      632.78373931320*t);
+   earth_x_1+=     0.00000000118 * cos( 5.96799701373 +    17654.78053974960*t);
+   earth_x_1+=     0.00000000105 * cos( 5.26799267002 +    23013.53953958720*t);
+   earth_x_1+=     0.00000000103 * cos( 0.24702432231 +     6271.34649715540*t);
+   earth_x_1+=     0.00000000103 * cos( 0.11832394316 +     6294.80520282740*t);
+   earth_x_1+=     0.00000000102 * cos( 2.69803619583 +    17267.26820169119*t);
+   earth_x_1+=     0.00000000108 * cos( 2.49639211639 +    23543.23050468179*t);
+   earth_x_1+=     0.00000000114 * cos( 2.21695343681 +     4171.42553661380*t);
+   earth_x_1+=     0.00000000101 * cos( 4.42100430236 +     4804.20927592700*t);
+   earth_x_1+=     0.00000000111 * cos( 3.66256237589 +     1581.95934828300*t);
+   earth_x_1+=     0.00000000119 * cos( 0.24456263546 +    10213.28554621100*t);
+   earth_x_1+=     0.00000000102 * cos( 5.67042027327 +     6549.68289171320*t);
+   earth_x_1+=     0.00000000102 * cos( 0.97811329938 +     6016.46880826960*t);
+   earth_x_1+=     0.00000000121 * cos( 0.05336972054 +    13625.53363017200*t);
+   earth_x_1+=     0.00000000090 * cos( 3.68828039857 +       76.26607127560*t);
+   earth_x_1+=     0.00000000086 * cos( 5.93949751968 +    12808.88030395680*t);
+   earth_x_1+=     0.00000000095 * cos( 1.34358559692 +    12345.73905754400*t);
+   earth_x_1+=     0.00000000106 * cos( 5.14281705612 +    11216.28429032400*t);
+   earth_x_1+=     0.00000000095 * cos( 2.50046439714 +    12012.58229714040*t);
+   earth_x_1+=     0.00000000081 * cos( 3.84772814821 +    11614.43329373220*t);
+   earth_x_1+=     0.00000000084 * cos( 1.14755262027 +     6255.67429389460*t);
+   earth_x_1+=     0.00000000084 * cos( 5.50098095238 +     6310.47740608820*t);
+   earth_x_1+=     0.00000000106 * cos( 6.00471264495 +     3097.88382272579*t);
+   earth_x_1+=     0.00000000080 * cos( 0.34711727404 +     6922.97313630540*t);
+   earth_x_1+=     0.00000000085 * cos( 5.79536877725 +     5636.06501667660*t);
+   earth_x_1+=     0.00000000092 * cos( 0.87789019103 +    10818.13528691580*t);
+   earth_x_1+=     0.00000000080 * cos( 4.19702650261 +     3583.34103067380*t);
+   earth_x_1+=     0.00000000083 * cos( 3.70902039500 +     5120.60114558360*t);
+   earth_x_1+=     0.00000000080 * cos( 1.41939312484 +    13521.75144159140*t);
+   earth_x_1+=     0.00000000084 * cos( 3.01891268791 +    10440.27429260360*t);
+   earth_x_1+=     0.00000000077 * cos( 5.70306938809 +     6993.00889854970*t);
+   earth_x_1+=     0.00000000104 * cos( 5.42105253386 +     5650.29211067820*t);
+   earth_x_1+=     0.00000000086 * cos( 0.81273058913 +    12359.96615154560*t);
+   earth_x_1+=     0.00000000090 * cos( 6.02294397735 +     6195.76764545159*t);
+   earth_x_1+=     0.00000000090 * cos( 0.62558959530 +     6370.38405453120*t);
+   earth_x_1+=     0.00000000074 * cos( 2.33758314304 +     3340.61242669980*t);
+   earth_x_1+=     0.00000000082 * cos( 0.19767948100 +    18842.11400297339*t);
+   earth_x_1+=     0.00000000083 * cos( 5.31994916859 +    18415.51581209739*t);
+   earth_x_1+=     0.00000000101 * cos( 2.86191811970 +     3104.93005942380*t);
+   earth_x_1+=     0.00000000073 * cos( 6.09594424715 +      949.17560896980*t);
+   earth_x_1+=     0.00000000084 * cos( 4.35521501162 +     2379.16447357160*t);
+   earth_x_1+=     0.00000000075 * cos( 1.62833922577 +    30355.99731976780*t);
+   earth_x_1+=     0.00000000069 * cos( 3.24662354142 +    10454.50138660520*t);
+   earth_x_1+=     0.00000000076 * cos( 3.92665935350 +    17298.18232732620*t);
+   earth_x_1+=     0.00000000076 * cos( 4.83482959872 +    18429.74290609900*t);
+   earth_x_1+=     0.00000000079 * cos( 0.06205634713 +    22805.73556599360*t);
+   earth_x_1+=     0.00000000077 * cos( 2.05455640378 +    16858.48253293320*t);
+   earth_x_1+=     0.00000000087 * cos( 5.09755111164 +     3903.91137641980*t);
+   earth_x_1+=     0.00000000061 * cos( 3.70271756349 +    16737.57723659660*t);
+   earth_x_1+=     0.00000000078 * cos( 4.31638741861 +     8671.96987044060*t);
+   earth_x_1+=     0.00000000063 * cos( 0.99541849434 +    24356.78078864160*t);
+   earth_x_1+=     0.00000000075 * cos( 2.14136838637 +    23581.25817731760*t);
+   earth_x_1+=     0.00000000073 * cos( 4.49527510239 +     9388.00590941520*t);
+   earth_x_1+=     0.00000000058 * cos( 0.89406841198 +    17996.03116822220*t);
+   earth_x_1+=     0.00000000057 * cos( 2.48914084100 +      419.48464387520*t);
+   earth_x_1+=     0.00000000062 * cos( 4.06342805588 +    11925.27409260060*t);
+   earth_x_1+=     0.00000000060 * cos( 3.51548088474 +    34570.06633485260*t);
+   earth_x_1+=     0.00000000057 * cos( 4.46050926275 +   149854.40013480789*t);
+   earth_x_1+=     0.00000000056 * cos( 5.33393809517 +    12455.94537876340*t);
+   earth_x_1+=     0.00000000056 * cos( 2.17965993818 +    11764.33076885900*t);
+   earth_x_1+=     0.00000000061 * cos( 0.15194461973 +    28237.23345938940*t);
+   earth_x_1+=     0.00000000065 * cos( 5.87045161229 +     7335.34423317980*t);
+   earth_x_1+=     0.00000000072 * cos( 5.87296425674 +     6510.55198278040*t);
+   earth_x_1+=     0.00000000072 * cos( 0.77556931591 +     6055.59971720240*t);
+   earth_x_1+=     0.00000000071 * cos( 0.23083621209 +    12545.37630449040*t);
+   earth_x_1+=     0.00000000061 * cos( 5.37705181537 +     6253.64634145540*t);
+   earth_x_1+=     0.00000000061 * cos( 1.27148175728 +     6312.50535852740*t);
+   earth_x_1+=     0.00000000055 * cos( 1.18603703570 +     2699.73481931760*t);
+   earth_x_1+=     0.00000000056 * cos( 6.18401368339 +     5573.14280143310*t);
+   earth_x_1+=     0.00000000055 * cos( 5.94073604896 +    15671.08175940660*t);
+   earth_x_1+=     0.00000000052 * cos( 0.26563539971 +     5326.78669402080*t);
+   earth_x_1+=     0.00000000053 * cos( 1.77739578512 +    26735.94526221320*t);
+   earth_x_1+=     0.00000000065 * cos( 2.33938496802 +    14919.01785375460*t);
+   earth_x_1+=     0.00000000049 * cos( 4.18739781015 +    35579.69123957000*t);
+   earth_x_1+=     0.00000000053 * cos( 5.28549359722 +     6357.85744855870*t);
+   earth_x_1+=     0.00000000053 * cos( 1.36303997543 +     6208.29425142410*t);
+   earth_x_1+=     0.00000000051 * cos( 0.76407742618 +    16522.65971600220*t);
+   earth_x_1+=     0.00000000049 * cos( 4.92674479666 +    21393.54196985760*t);
+   earth_x_1+=     0.00000000049 * cos( 0.60525485888 +     7018.95236352320*t);
+   earth_x_1+=     0.00000000049 * cos( 6.04327871377 +     5547.19933645960*t);
+   earth_x_1+=     0.00000000049 * cos( 2.59123701567 +    12416.58850284820*t);
+   earth_x_1+=     0.00000000050 * cos( 0.07523239468 +    26084.02180621620*t);
+   earth_x_1+=     0.00000000048 * cos( 5.46985443988 +    12132.43996210600*t);
+   earth_x_1+=     0.00000000051 * cos( 3.40683294217 +    23006.42599258639*t);
+   earth_x_1+=     0.00000000058 * cos( 2.46543150418 +    15141.39079431200*t);
+   earth_x_1+=     0.00000000048 * cos( 4.88190521609 +    16200.77272450120*t);
+   earth_x_1+=     0.00000000047 * cos( 0.97667389279 +    24279.10701821359*t);
+   earth_x_1+=     0.00000000049 * cos( 1.66846784739 +     5749.86176654780*t);
+   earth_x_1+=     0.00000000049 * cos( 4.98006572526 +     6816.28993343500*t);
+   earth_x_1+=     0.00000000046 * cos( 3.45134537804 +    19378.91851506880*t);
+   earth_x_1+=     0.00000000047 * cos( 2.09300174029 +    24065.80792277559*t);
+   earth_x_1+=     0.00000000046 * cos( 5.73797450552 +   104351.61256629678*t);
+   earth_x_1+=     0.00000000049 * cos( 1.45701599108 +     8982.81066930900*t);
+   earth_x_1+=     0.00000000050 * cos( 4.49441848554 +    24602.61243487099*t);
+   earth_x_1+=     0.00000000045 * cos( 0.12799634054 +    24491.42579258340*t);
+   earth_x_1+=     0.00000000049 * cos( 2.95918505969 +    16460.33352952499*t);
+   earth_x_1+=     0.00000000054 * cos( 0.03974674424 +     7834.12107263940*t);
+   earth_x_1+=     0.00000000045 * cos( 5.05860612319 +   161000.68573767410*t);
+   earth_x_1+=     0.00000000044 * cos( 3.92726741188 +     8662.24032356300*t);
+   earth_x_1+=     0.00000000049 * cos( 4.51486356794 +    29303.72893657940*t);
+   earth_x_1+=     0.00000000049 * cos( 4.98978743446 +    29289.50184257780*t);
+   earth_x_1+=     0.00000000043 * cos( 3.34677142688 +     5244.04923920100*t);
+   earth_x_1+=     0.00000000044 * cos( 3.63764737159 +     6175.41232605280*t);
+   earth_x_1+=     0.00000000044 * cos( 3.01088620106 +     6390.73937393000*t);
+   earth_x_1+=     0.00000000042 * cos( 1.47305743112 +    24705.70520908959*t);
+   earth_x_1+=     0.00000000041 * cos( 4.78696656856 +    17370.36097590980*t);
+   earth_x_1+=     0.00000000041 * cos( 4.99610869214 +    22483.84857449259*t);
+   earth_x_1+=     0.00000000057 * cos( 3.19497095530 +    84334.66158130829*t);
+   earth_x_1+=     0.00000000041 * cos( 3.62603303245 +    23020.65308658799*t);
+   earth_x_1+=     0.00000000045 * cos( 2.13147646951 +    54247.41313575040*t);
+   earth_x_1+=     0.00000000040 * cos( 3.00393229632 +    32243.49850088060*t);
+   earth_x_1+=     0.00000000040 * cos( 0.45946989377 +    24336.00539314920*t);
+   earth_x_1+=     0.00000000040 * cos( 0.00254937606 +     8401.83971036980*t);
+   earth_x_1+=     0.00000000041 * cos( 5.05922003836 +    26482.17080962440*t);
+   earth_x_1+=     0.00000000044 * cos( 4.93011217974 +    28230.18722269139*t);
+   earth_x_1+=     0.00000000039 * cos( 5.54872052106 +    45585.17281218740*t);
+   earth_x_1+=     0.00000000038 * cos( 3.93561558747 +     6268.84875598980*t);
+   earth_x_1+=     0.00000000038 * cos( 2.71291798518 +     6297.30294399300*t);
+   earth_x_1+=     0.00000000041 * cos( 4.95311845784 +    15508.61512327440*t);
+   earth_x_1+=     0.00000000038 * cos( 0.10976793876 +     6171.88920770380*t);
+   earth_x_1+=     0.00000000038 * cos( 0.25558032671 +     6394.26249227900*t);
+   earth_x_1+=     0.00000000039 * cos( 3.31511780881 +    29826.30635467320*t);
+   earth_x_1+=     0.00000000038 * cos( 4.03952051907 +    16062.18452611680*t);
+   earth_x_1+=     0.00000000037 * cos( 2.07587316842 +    20452.86941222180*t);
+   earth_x_1+=     0.00000000049 * cos( 5.78169307611 +    18216.44381066100*t);
+   earth_x_1+=     0.00000000041 * cos( 1.35496527372 +    29088.81141598500*t);
+   earth_x_1+=     0.00000000038 * cos( 6.16268217428 +    18202.21671665939*t);
+   earth_x_1+=     0.00000000036 * cos( 3.84319801860 +    39302.09696219600*t);
+   earth_x_1+=     0.00000000038 * cos( 6.14920442251 +    26709.64694241340*t);
+   earth_x_1+=     0.00000000036 * cos( 6.28049708584 +     9924.81042151060*t);
+   earth_x_1+=     0.00000000035 * cos( 4.66732134451 +     4061.21921539440*t);
+   earth_x_1+=     0.00000000040 * cos( 3.43223548088 +    27511.46787353720*t);
+   earth_x_1+=     0.00000000034 * cos( 1.38801022558 +     6147.45052498140*t);
+   earth_x_1+=     0.00000000034 * cos( 5.26052334707 +     6418.70117500140*t);
+   earth_x_1+=     0.00000000034 * cos( 1.97967478962 +    17271.88400752640*t);
+   earth_x_1+=     0.00000000037 * cos( 6.05744777692 +    21150.81336588360*t);
+   earth_x_1+=     0.00000000034 * cos( 1.90509484306 +    29026.48522950779*t);
+   earth_x_1+=     0.00000000035 * cos( 3.20661690466 +     7322.10246078180*t);
+   earth_x_1+=     0.00000000044 * cos( 2.79017179925 +     3981.49003408200*t);
+   earth_x_1+=     0.00000000034 * cos( 1.21719184072 +    21548.96236929180*t);
+   earth_x_1+=     0.00000000033 * cos( 4.13127114332 +   173566.83743765688*t);
+   earth_x_1+=     0.00000000033 * cos( 3.22867201822 +    30666.15495843280*t);
+   earth_x_1+=     0.00000000034 * cos( 4.95128499531 +    38500.27603107220*t);
+   earth_x_1+=     0.00000000032 * cos( 4.38302389502 +    17686.75284556640*t);
+   earth_x_1+=     0.00000000038 * cos( 4.95582614524 +    19004.64794940840*t);
+   earth_x_1+=     0.00000000033 * cos( 1.70921741352 +    34520.30930938080*t);
+   earth_x_1+=     0.00000000041 * cos( 4.09101902227 +    21202.09370374600*t);
+   earth_x_1+=     0.00000000031 * cos( 0.28544969417 +    22003.91463486980*t);
+   earth_x_1+=     0.00000000037 * cos( 1.21597687953 +     7349.57132718140*t);
+   earth_x_1+=     0.00000000031 * cos( 0.19462514993 +    27832.03821928320*t);
+   earth_x_1+=     0.00000000037 * cos( 3.41912843249 +    29296.61538957860*t);
+   earth_x_1+=     0.00000000035 * cos( 2.69261244483 +    11919.14086666800*t);
+   earth_x_1+=     0.00000000035 * cos( 0.83074174405 +    29062.51309618520*t);
+   earth_x_1+=     0.00000000030 * cos( 2.96012497927 +    12489.88562870720*t);
+   earth_x_1+=     0.00000000035 * cos( 4.83736876925 +    21228.39202354580*t);
+   earth_x_1+=     0.00000000032 * cos( 6.26164802166 +     6231.79551212899*t);
+   earth_x_1+=     0.00000000032 * cos( 0.38688555099 +     6334.35618785381*t);
+   earth_x_1+=     0.00000000029 * cos( 0.70432070922 +    49515.38250840700*t);
+   earth_x_1+=     0.00000000032 * cos( 0.44788919628 +    31172.65064598300*t);
+   earth_x_1+=     0.00000000031 * cos( 2.69970301246 +    33019.02111220460*t);
+   earth_x_1+=     0.00000000029 * cos( 2.26256470092 +    27278.46881644080*t);
+   earth_x_1+=     0.00000000031 * cos( 3.19023486538 +     5017.50837136500*t);
+   earth_x_1+=     0.00000000028 * cos( 1.87849143707 +    24080.03501677719*t);
+   earth_x_1+=     0.00000000034 * cos( 0.78574341965 +     5905.70224207560*t);
+   earth_x_1+=     0.00000000034 * cos( 5.86279015300 +     6660.44945790720*t);
+   earth_x_1+=     0.00000000038 * cos( 5.12224209756 +    31570.79964939120*t);
+   earth_x_1+=     0.00000000027 * cos( 3.76025315460 +    19402.79695281660*t);
+   earth_x_1+=     0.00000000027 * cos( 0.53327304803 +   233141.31440436149*t);
+   earth_x_1+=     0.00000000030 * cos( 2.32405719366 +    12586.92709547520*t);
+   earth_x_1+=     0.00000000028 * cos( 5.54772804461 +    34911.41207609100*t);
+   earth_x_1+=     0.00000000027 * cos( 5.30079330746 +    43232.30665841560*t);
+   earth_x_1+=     0.00000000027 * cos( 6.17189140209 +     1692.16566950240*t);
+   earth_x_1+=     0.00000000036 * cos( 4.81217077776 +    15664.03552270859*t);
+   earth_x_1+=     0.00000000027 * cos( 0.65637418130 +     6360.74962041940*t);
+   earth_x_1+=     0.00000000027 * cos( 5.99215939135 +     6205.40207956340*t);
+   earth_x_1+=     0.00000000026 * cos( 3.42413737440 +    24383.07910844140*t);
+   earth_x_1+=     0.00000000033 * cos( 3.45722382528 +    28628.33622609960*t);
+   earth_x_1+=     0.00000000027 * cos( 2.32049228976 +    32367.09765620760*t);
+   earth_x_1+=     0.00000000026 * cos( 1.98182489066 +    19804.82729158280*t);
+   earth_x_1+=     0.00000000027 * cos( 1.77114813957 +     4487.81740627040*t);
+   earth_x_1+=     0.00000000028 * cos( 2.89944615107 +    58953.14544329400*t);
+   earth_x_1+=     0.00000000028 * cos( 3.80331477821 +    10239.58386601080*t);
+   earth_x_1+=     0.00000000026 * cos( 0.36912120822 +     6339.97422492700*t);
+   earth_x_1+=     0.00000000026 * cos( 6.27941236443 +     6226.17747505580*t);
+   earth_x_1+=     0.00000000026 * cos( 6.01416870867 +     4797.09572892620*t);
+   earth_x_1+=     0.00000000028 * cos( 2.80825729869 +    22476.73502749179*t);
+   earth_x_1+=     0.00000000029 * cos( 0.17746825116 +     6592.35417264720*t);
+   earth_x_1+=     0.00000000029 * cos( 0.18788001431 +     5973.79752733560*t);
+   earth_x_1+=     0.00000000026 * cos( 3.93576097012 +    36949.23080842420*t);
+   earth_x_1+=     0.00000000027 * cos( 0.64682460545 +    29424.63423291600*t);
+   earth_x_1+=     0.00000000026 * cos( 2.54059552040 +     9910.58332750900*t);
+   earth_x_1+=     0.00000000028 * cos( 2.41855003101 +    46386.99374331120*t);
+   earth_x_1+=     0.00000000024 * cos( 1.22629501602 +    21947.11137270000*t);
+   earth_x_1+=     0.00000000024 * cos( 0.55027097678 +    11616.97609101300*t);
+   earth_x_1+=     0.00000000025 * cos( 2.78491065355 +     6302.74561089119*t);
+   earth_x_1+=     0.00000000025 * cos( 3.86362291910 +     6263.40608909160*t);
+   earth_x_1+=     0.00000000026 * cos( 4.44347824401 +    15110.46611986620*t);
+   earth_x_1+=     0.00000000027 * cos( 5.91935985488 +    52175.80628314840*t);
+   earth_x_1+=     0.00000000028 * cos( 6.10918595816 +    18052.92954315780*t);
+   earth_x_1+=     0.00000000023 * cos( 1.99242379288 +    30774.50164257480*t);
+   earth_x_1+=     0.00000000024 * cos( 3.73458431258 +    31968.94865279940*t);
+   earth_x_1+=     0.00000000023 * cos( 0.62347300029 +    17892.93839400359*t);
+   earth_x_1+=     0.00000000024 * cos( 3.21260068222 +    23550.34405168259*t);
+   earth_x_1+=     0.00000000023 * cos( 1.36063534536 +     6716.78758786820*t);
+   earth_x_1+=     0.00000000027 * cos( 0.02724821878 +    25158.60171976540*t);
+   earth_x_1+=     0.00000000026 * cos( 3.57645528088 +     8584.66166590080*t);
+   earth_x_1+=     0.00000000022 * cos( 3.99741791962 +     7445.55055439920*t);
+   earth_x_1+=     0.00000000022 * cos( 6.10694895695 +    15265.88651930040*t);
+   earth_x_1+=     0.00000000025 * cos( 0.81555419576 +    14155.22459526660*t);
+   earth_x_1+=     0.00000000023 * cos( 3.06283945185 +     6272.03014972750*t);
+   earth_x_1+=     0.00000000023 * cos( 3.58569412080 +     6294.12155025530*t);
+   earth_x_1+=     0.00000000025 * cos( 2.44880076940 +    20199.09495963300*t);
+   earth_x_1+=     0.00000000021 * cos( 5.78253540504 +    22779.43724619380*t);
+   earth_x_1+=     0.00000000023 * cos( 2.25096306087 +    34115.11406927460*t);
+   earth_x_1+=     0.00000000022 * cos( 4.67990936559 +    34596.36465465240*t);
+   earth_x_1+=     0.00000000029 * cos( 0.93434607315 +    11933.36796066960*t);
+   earth_x_1+=     0.00000000021 * cos( 4.08577407328 +    30220.93223973240*t);
+   earth_x_1+=     0.00000000020 * cos( 0.47093224249 +    47162.51635463520*t);
+   earth_x_1+=     0.00000000022 * cos( 1.38262024424 +     4379.63903749020*t);
+   earth_x_1+=     0.00000000020 * cos( 5.31355969163 +    40879.44050464380*t);
+   earth_x_1+=     0.00000000022 * cos( 4.52226925108 +    28759.81087748319*t);
+   earth_x_1+=     0.00000000026 * cos( 3.41054638032 +    18422.62935909819*t);
+   earth_x_1+=     0.00000000027 * cos( 0.67002293804 +    34513.26307268280*t);
+   earth_x_1+=     0.00000000019 * cos( 4.78610064760 +    28313.28880466100*t);
+   earth_x_1+=     0.00000000020 * cos( 5.09152714894 +     8186.51266249260*t);
+   earth_x_1+=     0.00000000025 * cos( 3.53211055104 +    55798.45835839840*t);
+   earth_x_1+=     0.00000000025 * cos( 0.38024667868 +     9225.53927328300*t);
+   earth_x_1+=     0.00000000020 * cos( 1.71238700404 +    27433.88921587499*t);
+   earth_x_1+=     0.00000000024 * cos( 1.20856889899 +     9411.46461508720*t);
+   earth_x_1+=     0.00000000020 * cos( 1.65292451759 +     6379.05507720920*t);
+   earth_x_1+=     0.00000000020 * cos( 4.99560905506 +     6187.09662277360*t);
+   earth_x_1+=     0.00000000020 * cos( 2.94178503134 +    19800.94595622480*t);
+   earth_x_1+=     0.00000000020 * cos( 4.32349571107 +    62883.35513951360*t);
+   earth_x_1+=     0.00000000022 * cos( 4.44470422678 +    18326.65013188039*t);
+   earth_x_1+=     0.00000000019 * cos( 4.37840329985 +    16207.88627150200*t);
+   earth_x_1+=     0.00000000022 * cos( 4.46505143635 +     6332.83287546320*t);
+   earth_x_1+=     0.00000000022 * cos( 2.18348213630 +     6233.31882451960*t);
+   earth_x_1+=     0.00000000019 * cos( 1.36569875427 +    23539.70738633280*t);
+   earth_x_1+=     0.00000000018 * cos( 4.61817062443 +     6593.91664885980*t);
+   earth_x_1+=     0.00000000018 * cos( 2.03036294822 +     5972.23505112300*t);
+   earth_x_1+=     0.00000000019 * cos( 2.36455752682 +     5618.31980486140*t);
+   earth_x_1+=     0.00000000019 * cos( 4.28397604583 +     6947.83189512140*t);
+   earth_x_1+=     0.00000000019 * cos( 0.18461307728 +    83973.83535572988*t);
+   earth_x_1+=     0.00000000019 * cos( 4.29574440213 +    84019.85928049390*t);
+   earth_x_1+=     0.00000000020 * cos( 2.85390720284 +     9070.11887384880*t);
+   earth_x_1=earth_x_1 * t;
+
+   double earth_x_2=0.0;
+   earth_x_2+=     0.00004143217 * cos( 3.14159265359 +        0.00000000000*t);
+   earth_x_2+=     0.00002175695 * cos( 4.39999849572 +    12566.15169998280*t);
+   earth_x_2+=     0.00000995233 * cos( 0.20790847155 +     6283.07584999140*t);
+   earth_x_2+=     0.00000092659 * cos( 4.19285471010 +    18849.22754997420*t);
+   earth_x_2+=     0.00000013679 * cos( 3.37833642063 +     6286.59896834040*t);
+   earth_x_2+=     0.00000013668 * cos( 3.27271492019 +     6279.55273164240*t);
+   earth_x_2+=     0.00000008190 * cos( 4.41928022030 +     6256.77753019160*t);
+   earth_x_2+=     0.00000008145 * cos( 2.22925093408 +     6309.37416979120*t);
+   earth_x_2+=     0.00000008004 * cos( 5.92288197328 +     6127.65545055720*t);
+   earth_x_2+=     0.00000007746 * cos( 0.72469358509 +     6438.49624942560*t);
+   earth_x_2+=     0.00000008899 * cos( 5.38113427234 +    83996.84731811189*t);
+   earth_x_2+=     0.00000004754 * cos( 6.16710003373 +      709.93304855830*t);
+   earth_x_2+=     0.00000003737 * cos( 2.48724420764 +     5507.55323866740*t);
+   earth_x_2+=     0.00000002893 * cos( 4.94972697460 +     4705.73230754360*t);
+   earth_x_2+=     0.00000002937 * cos( 4.08160813245 +    25132.30339996560*t);
+   earth_x_2+=     0.00000003203 * cos( 4.14275083744 +     7058.59846131540*t);
+   earth_x_2+=     0.00000002254 * cos( 3.51957213521 +      775.52261132400*t);
+   earth_x_2+=     0.00000002170 * cos( 1.70495739017 +     7860.41939243920*t);
+   earth_x_2+=     0.00000002031 * cos( 2.29363438970 +     6275.96230299060*t);
+   earth_x_2+=     0.00000002031 * cos( 4.35489918295 +     6290.18939699220*t);
+   earth_x_2+=     0.00000001904 * cos( 4.46668961834 +     5486.77784317500*t);
+   earth_x_2+=     0.00000001720 * cos( 5.23142094282 +     1059.38193018920*t);
+   earth_x_2+=     0.00000001548 * cos( 2.12906745749 +     6040.34724601740*t);
+   earth_x_2+=     0.00000001563 * cos( 2.18101487697 +     7079.37385680780*t);
+   earth_x_2+=     0.00000001456 * cos( 3.53770823599 +     5753.38488489680*t);
+   earth_x_2+=     0.00000001424 * cos( 3.16250984559 +     6812.76681508600*t);
+   earth_x_2+=     0.00000001470 * cos( 4.51659549398 +     6525.80445396540*t);
+   earth_x_2+=     0.00000001429 * cos( 3.01998118552 +     5884.92684658320*t);
+   earth_x_2+=     0.00000001288 * cos( 3.62892288697 +     6681.22485339960*t);
+   earth_x_2+=     0.00000001338 * cos( 1.04445705185 +     5729.50644714900*t);
+   earth_x_2+=     0.00000001039 * cos( 5.85310113863 +     6282.09552892320*t);
+   earth_x_2+=     0.00000001039 * cos( 0.79543243402 +     6284.05617105960*t);
+   earth_x_2+=     0.00000001212 * cos( 5.61713170172 +     6836.64525283380*t);
+   earth_x_2+=     0.00000000904 * cos( 2.43535199673 +     5331.35744374080*t);
+   earth_x_2+=     0.00000000871 * cos( 3.84816000443 +     4933.20844033260*t);
+   earth_x_2+=     0.00000000879 * cos( 3.06133930098 +    11790.62908865880*t);
+   earth_x_2+=     0.00000001102 * cos( 4.59510736976 +   167283.76158766549*t);
+   earth_x_2+=     0.00000000769 * cos( 4.22812146942 +     7234.79425624200*t);
+   earth_x_2+=     0.00000000744 * cos( 0.34956564055 +     4136.91043351620*t);
+   earth_x_2+=     0.00000000939 * cos( 3.62288542048 +    11856.21865142450*t);
+   earth_x_2+=     0.00000000675 * cos( 2.82587446878 +     7632.94325965020*t);
+   earth_x_2+=     0.00000000654 * cos( 1.40322441961 +    11506.76976979360*t);
+   earth_x_2+=     0.00000000854 * cos( 5.24328642976 +     4535.05943692440*t);
+   earth_x_2+=     0.00000000665 * cos( 4.81749499447 +     1194.44701022460*t);
+   earth_x_2+=     0.00000000834 * cos( 5.75118975782 +      155.42039943420*t);
+   earth_x_2+=     0.00000000785 * cos( 5.88776299473 +     5088.62883976680*t);
+   earth_x_2+=     0.00000000548 * cos( 6.12128965864 +      796.29800681640*t);
+   earth_x_2+=     0.00000000686 * cos( 1.90135823927 +      801.82093112380*t);
+   earth_x_2+=     0.00000000507 * cos( 5.42408029291 +     4694.00295470760*t);
+   earth_x_2+=     0.00000000677 * cos( 2.71965226391 +     3128.38876509580*t);
+   earth_x_2+=     0.00000000499 * cos( 3.02752235850 +    12569.67481833180*t);
+   earth_x_2+=     0.00000000480 * cos( 0.01327044252 +     8429.24126646660*t);
+   earth_x_2+=     0.00000000592 * cos( 1.39433962245 +     8031.09226305840*t);
+   earth_x_2+=     0.00000000498 * cos( 6.19000103499 +     1589.07289528380*t);
+   earth_x_2+=     0.00000000562 * cos( 0.77105088912 +     7477.52286021600*t);
+   earth_x_2+=     0.00000000415 * cos( 0.00558962450 +     6069.77675455340*t);
+   earth_x_2+=     0.00000000415 * cos( 0.35131226280 +     6496.37494542940*t);
+   earth_x_2+=     0.00000000402 * cos( 1.72177128440 +     3738.76143010800*t);
+   earth_x_2+=     0.00000000503 * cos( 3.41456622371 +     1592.59601363280*t);
+   earth_x_2+=     0.00000000390 * cos( 5.36711219721 +       26.29831979980*t);
+   earth_x_2+=     0.00000000372 * cos( 4.35153585802 +      553.56940284240*t);
+   earth_x_2+=     0.00000000458 * cos( 0.85618925594 +   156137.47598479928*t);
+   earth_x_2+=     0.00000000325 * cos( 2.94948293011 +      951.71840625060*t);
+   earth_x_2+=     0.00000000321 * cos( 2.01798251637 +     1990.74501704100*t);
+   earth_x_2+=     0.00000000335 * cos( 1.12608331535 +     4690.47983635860*t);
+   earth_x_2+=     0.00000000366 * cos( 0.10449761483 +    12721.57209941700*t);
+   earth_x_2+=     0.00000000287 * cos( 2.26476426223 +      242.72860397400*t);
+   earth_x_2+=     0.00000000347 * cos( 1.34248948630 +     5481.25491886760*t);
+   earth_x_2+=     0.00000000368 * cos( 3.91910931412 +     9437.76293488700*t);
+   earth_x_2+=     0.00000000255 * cos( 1.83876117314 +    11371.70468975820*t);
+   earth_x_2+=     0.00000000242 * cos( 4.92327001310 +     8827.39026987480*t);
+   earth_x_2+=     0.00000000307 * cos( 1.56817330160 +     1349.86740965880*t);
+   earth_x_2+=     0.00000000242 * cos( 0.12263961387 +     1748.01641306700*t);
+   earth_x_2+=     0.00000000222 * cos( 4.88851960168 +     4164.31198961300*t);
+   earth_x_2+=     0.00000000230 * cos( 1.05870192923 +     3154.68708489560*t);
+   earth_x_2+=     0.00000000204 * cos( 0.44240387853 +    11769.85369316640*t);
+   earth_x_2+=     0.00000000188 * cos( 5.65888993470 +     7875.67186362420*t);
+   earth_x_2+=     0.00000000189 * cos( 0.61654841817 +     2388.89402044920*t);
+   earth_x_2+=     0.00000000233 * cos( 5.23023625373 +     7084.89678111520*t);
+   earth_x_2+=     0.00000000183 * cos( 4.10128236867 +     7342.45778018060*t);
+   earth_x_2+=     0.00000000173 * cos( 0.52511015365 +    10977.07880469900*t);
+   earth_x_2+=     0.00000000159 * cos( 2.72294432534 +     4292.33083295040*t);
+   earth_x_2+=     0.00000000207 * cos( 2.28587256917 +    13367.97263110660*t);
+   earth_x_2+=     0.00000000202 * cos( 3.21445201981 +    10973.55568635000*t);
+   earth_x_2+=     0.00000000145 * cos( 5.45882589245 +      398.14900340820*t);
+   earth_x_2+=     0.00000000155 * cos( 5.81412827494 +    12410.73130054860*t);
+   earth_x_2+=     0.00000000134 * cos( 4.56296979111 +    10575.40668294180*t);
+   earth_x_2+=     0.00000000127 * cos( 4.84612567560 +    13119.72110282519*t);
+   earth_x_2+=     0.00000000115 * cos( 4.96539448045 +     2146.16541647520*t);
+   earth_x_2+=     0.00000000125 * cos( 3.87814283336 +    12562.62858163380*t);
+   earth_x_2+=     0.00000000125 * cos( 0.17913722961 +     8635.94200376320*t);
+   earth_x_2+=     0.00000000122 * cos( 1.05019176110 +    12592.45001978260*t);
+   earth_x_2+=     0.00000000111 * cos( 2.15765835240 +     1577.34354244780*t);
+   earth_x_2+=     0.00000000134 * cos( 6.25201667258 +     3930.20969621960*t);
+   earth_x_2+=     0.00000000107 * cos( 3.43601884429 +    13517.87010623340*t);
+   earth_x_2+=     0.00000000102 * cos( 4.37125144267 +    10988.80815753500*t);
+   earth_x_2+=     0.00000000095 * cos( 5.98656941724 +    17260.15465469040*t);
+   earth_x_2+=     0.00000000100 * cos( 2.34969093002 +    13095.84266507740*t);
+   earth_x_2+=     0.00000000110 * cos( 3.63768958070 +     5223.69391980220*t);
+   earth_x_2+=     0.00000000085 * cos( 4.32512089859 +     3894.18182954220*t);
+   earth_x_2+=     0.00000000109 * cos( 3.99726753431 +    12539.85338018300*t);
+   earth_x_2+=     0.00000000079 * cos( 5.22586609526 +    16730.46368959580*t);
+   earth_x_2+=     0.00000000079 * cos( 0.08497701397 +    10447.38783960440*t);
+   earth_x_2+=     0.00000000076 * cos( 4.11846374330 +    12323.42309600880*t);
+   earth_x_2+=     0.00000000105 * cos( 3.53845510796 +     5642.19824260920*t);
+   earth_x_2+=     0.00000000075 * cos( 5.73067319226 +     6709.67404086740*t);
+   earth_x_2+=     0.00000000075 * cos( 0.91786038039 +     5856.47765911540*t);
+   earth_x_2+=     0.00000000082 * cos( 5.90914298142 +    10177.25767953360*t);
+   earth_x_2+=     0.00000000100 * cos( 0.66749297482 +   250570.67585721909*t);
+   earth_x_2+=     0.00000000071 * cos( 0.59970099558 +    14314.16811304980*t);
+   earth_x_2+=     0.00000000092 * cos( 1.98240564470 +    13916.01910964160*t);
+   earth_x_2+=     0.00000000091 * cos( 3.10960444681 +     6923.95345737360*t);
+   earth_x_2+=     0.00000000082 * cos( 4.00218490582 +    31415.37924995700*t);
+   earth_x_2+=     0.00000000083 * cos( 5.97854162397 +    95143.13292097810*t);
+   earth_x_2+=     0.00000000061 * cos( 4.23984291228 +     8273.82086703240*t);
+   earth_x_2+=     0.00000000055 * cos( 1.41429091428 +     5643.17856367740*t);
+   earth_x_2+=     0.00000000057 * cos( 4.71394484871 +      639.89728631400*t);
+   earth_x_2+=     0.00000000054 * cos( 5.63153101737 +    16496.36139620240*t);
+   earth_x_2+=     0.00000000075 * cos( 5.01329496789 +    71430.69561812909*t);
+   earth_x_2+=     0.00000000055 * cos( 3.83548544187 +    13341.67431130680*t);
+   earth_x_2+=     0.00000000052 * cos( 4.39839251424 +     3634.62102451840*t);
+   earth_x_2+=     0.00000000056 * cos( 1.31106242697 +    15720.83878487840*t);
+   earth_x_2+=     0.00000000050 * cos( 4.18716262207 +    12573.26524698360*t);
+   earth_x_2+=     0.00000000064 * cos( 3.35781526441 +     4732.03062734340*t);
+   earth_x_2+=     0.00000000066 * cos( 5.49585898905 +     2787.04302385740*t);
+   earth_x_2+=     0.00000000048 * cos( 2.02353005537 +    13362.44970679920*t);
+   earth_x_2+=     0.00000000049 * cos( 2.48296676043 +    14143.49524243060*t);
+   earth_x_2+=     0.00000000053 * cos( 4.12375134867 +    14919.01785375460*t);
+   earth_x_2+=     0.00000000047 * cos( 3.82959048333 +    12964.30070339100*t);
+   earth_x_2+=     0.00000000042 * cos( 0.38069707842 +    24072.92146977640*t);
+   earth_x_2+=     0.00000000051 * cos( 1.26573618407 +     6195.76764545159*t);
+   earth_x_2+=     0.00000000051 * cos( 5.38279738858 +     6370.38405453120*t);
+   earth_x_2+=     0.00000000048 * cos( 0.59989438773 +     2118.76386037840*t);
+   earth_x_2+=     0.00000000041 * cos( 2.49731951011 +    18319.53658487960*t);
+   earth_x_2+=     0.00000000041 * cos( 4.09328168773 +     3185.19202726560*t);
+   earth_x_2+=     0.00000000056 * cos( 1.01626259839 +    84672.47584450469*t);
+   earth_x_2+=     0.00000000039 * cos( 4.27641929836 +    20426.57109242200*t);
+   earth_x_2+=     0.00000000053 * cos( 3.21190793912 +   239424.39025435288*t);
+   earth_x_2+=     0.00000000040 * cos( 6.05039617409 +     3496.03282613400*t);
+   earth_x_2+=     0.00000000044 * cos( 1.67103342381 +     8662.24032356300*t);
+   earth_x_2+=     0.00000000043 * cos( 3.69693055127 +    18139.29450141590*t);
+   earth_x_2+=     0.00000000035 * cos( 3.60804630271 +    18875.52586977400*t);
+   earth_x_2+=     0.00000000035 * cos( 6.09077128153 +    10419.98628350760*t);
+   earth_x_2+=     0.00000000033 * cos( 0.53088739877 +    10213.28554621100*t);
+   earth_x_2+=     0.00000000031 * cos( 0.64652215197 +    23013.53953958720*t);
+   earth_x_2+=     0.00000000035 * cos( 0.91150627653 +     9779.10867612540*t);
+   earth_x_2+=     0.00000000030 * cos( 4.39715700541 +    12808.88030395680*t);
+   earth_x_2+=     0.00000000042 * cos( 5.61217219717 +    82576.98122099529*t);
+   earth_x_2+=     0.00000000031 * cos( 5.53053428232 +    14712.31711645800*t);
+   earth_x_2+=     0.00000000041 * cos( 5.62875780162 +    17789.84561978500*t);
+   earth_x_2+=     0.00000000031 * cos( 0.76672410297 +    22779.43724619380*t);
+   earth_x_2+=     0.00000000028 * cos( 2.34481351815 +     5429.87946823940*t);
+   earth_x_2+=     0.00000000027 * cos( 4.68961132588 +    16200.77272450120*t);
+   earth_x_2+=     0.00000000028 * cos( 3.91086715449 +    11015.10647733480*t);
+   earth_x_2+=     0.00000000026 * cos( 2.29393980670 +     9623.68827669120*t);
+   earth_x_2+=     0.00000000026 * cos( 4.35459376595 +     2942.46342329160*t);
+   earth_x_2+=     0.00000000026 * cos( 1.75234208326 +    18209.33026366019*t);
+   earth_x_2+=     0.00000000025 * cos( 6.04606961258 +     9917.69687450980*t);
+   earth_x_2+=     0.00000000025 * cos( 2.68813215214 +     3583.34103067380*t);
+   earth_x_2+=     0.00000000026 * cos( 0.90493478768 +    12012.58229714040*t);
+   earth_x_2+=     0.00000000023 * cos( 2.29788589227 +     9380.95967271720*t);
+   earth_x_2+=     0.00000000023 * cos( 2.49168005831 +    12559.03815298200*t);
+   earth_x_2+=     0.00000000023 * cos( 2.26727226510 +    26709.64694241340*t);
+   earth_x_2+=     0.00000000024 * cos( 4.92046069591 +    77713.77146812050*t);
+   earth_x_2+=     0.00000000022 * cos( 5.99172814093 +     6993.00889854970*t);
+   earth_x_2+=     0.00000000021 * cos( 2.38334945179 +    17256.63153634140*t);
+   earth_x_2+=     0.00000000023 * cos( 0.80841468126 +   149854.40013480789*t);
+   earth_x_2+=     0.00000000026 * cos( 3.19870156016 +    84334.66158130829*t);
+   earth_x_2+=     0.00000000023 * cos( 2.53410224481 +    19651.04848109800*t);
+   earth_x_2+=     0.00000000019 * cos( 1.93058658449 +    11926.25441366880*t);
+   earth_x_2+=     0.00000000018 * cos( 2.28255799994 +    11614.43329373220*t);
+   earth_x_2+=     0.00000000019 * cos( 3.66100662326 +    30639.85663863300*t);
+   earth_x_2+=     0.00000000022 * cos( 3.76382094449 +    17298.18232732620*t);
+   earth_x_2+=     0.00000000019 * cos( 5.39096925285 +    16460.33352952499*t);
+   earth_x_2+=     0.00000000023 * cos( 4.78353231767 +    72850.56171524568*t);
+   earth_x_2+=     0.00000000018 * cos( 3.19936496001 +    18073.70493865020*t);
+   earth_x_2+=     0.00000000017 * cos( 3.95703025000 +     8982.81066930900*t);
+   earth_x_2+=     0.00000000021 * cos( 3.77368081700 +    16858.48253293320*t);
+   earth_x_2+=     0.00000000021 * cos( 4.47047819380 +    12036.46073488820*t);
+   earth_x_2+=     0.00000000015 * cos( 0.65562803693 +    16062.18452611680*t);
+   earth_x_2+=     0.00000000014 * cos( 4.12927933278 +     2352.86615377180*t);
+   earth_x_2+=     0.00000000015 * cos( 4.21849500083 +    20199.09495963300*t);
+   earth_x_2+=     0.00000000014 * cos( 0.15947511180 +    22483.84857449259*t);
+   earth_x_2+=     0.00000000019 * cos( 2.81419648175 +     8671.96987044060*t);
+   earth_x_2+=     0.00000000014 * cos( 1.26341547558 +    90279.92316810328*t);
+   earth_x_2+=     0.00000000016 * cos( 0.41160329502 +    23543.23050468179*t);
+   earth_x_2+=     0.00000000012 * cos( 2.70990889006 +    17996.03116822220*t);
+   earth_x_2+=     0.00000000012 * cos( 0.72222975090 +    19402.79695281660*t);
+   earth_x_2+=     0.00000000012 * cos( 4.37880728616 +    10818.13528691580*t);
+   earth_x_2+=     0.00000000015 * cos( 0.57625721169 +    12567.13202105100*t);
+   earth_x_2+=     0.00000000011 * cos( 1.78472830783 +    28766.92442448400*t);
+   earth_x_2+=     0.00000000013 * cos( 3.18567583228 +     5573.14280143310*t);
+   earth_x_2+=     0.00000000011 * cos( 4.93327896664 +    18451.07854656599*t);
+   earth_x_2+=     0.00000000011 * cos( 5.43587992260 +    14945.31617355440*t);
+   earth_x_2+=     0.00000000011 * cos( 2.24607508630 +    35371.88726597640*t);
+   earth_x_2+=     0.00000000011 * cos( 5.02228901732 +    22805.73556599360*t);
+   earth_x_2+=     0.00000000010 * cos( 2.16930383321 +    15664.03552270859*t);
+   earth_x_2+=     0.00000000014 * cos( 1.93314444705 +    21228.39202354580*t);
+   earth_x_2+=     0.00000000010 * cos( 6.05310595892 +    23141.55838292460*t);
+   earth_x_2+=     0.00000000010 * cos( 2.55645282698 +    22345.26037610820*t);
+   earth_x_2+=     0.00000000009 * cos( 1.49155179482 +    29826.30635467320*t);
+   earth_x_2+=     0.00000000010 * cos( 0.18358815992 +    12779.45079542080*t);
+   earth_x_2+=     0.00000000011 * cos( 4.05473090705 +    23581.25817731760*t);
+   earth_x_2+=     0.00000000011 * cos( 0.39656469738 +    25158.60171976540*t);
+   earth_x_2=earth_x_2 * t * t;
+
+   double earth_x_3=0.0;
+   earth_x_3+=     0.00000175213 * cos( 3.14159265359 +        0.00000000000*t);
+   earth_x_3+=     0.00000072337 * cos( 2.89303952476 +    12566.15169998280*t);
+   earth_x_3+=     0.00000008364 * cos( 3.85500954096 +     6283.07584999140*t);
+   earth_x_3+=     0.00000005040 * cos( 2.53152989786 +    18849.22754997420*t);
+   earth_x_3+=     0.00000001481 * cos( 2.23672235234 +     6438.49624942560*t);
+   earth_x_3+=     0.00000001481 * cos( 4.41181122031 +     6127.65545055720*t);
+   earth_x_3+=     0.00000000644 * cos( 4.88567904780 +     6279.55273164240*t);
+   earth_x_3+=     0.00000000644 * cos( 1.76285452485 +     6286.59896834040*t);
+   earth_x_3+=     0.00000000318 * cos( 3.63840310183 +     6040.34724601740*t);
+   earth_x_3+=     0.00000000318 * cos( 3.01013047082 +     6525.80445396540*t);
+   earth_x_3+=     0.00000000220 * cos( 2.34166945312 +    25132.30339996560*t);
+   earth_x_3+=     0.00000000201 * cos( 0.82714760810 +     6836.64525283380*t);
+   earth_x_3+=     0.00000000201 * cos( 5.82138596455 +     5729.50644714900*t);
+   earth_x_3+=     0.00000000036 * cos( 1.41202770454 +    12569.67481833180*t);
+   earth_x_3+=     0.00000000033 * cos( 4.23554585611 +    12410.73130054860*t);
+   earth_x_3+=     0.00000000018 * cos( 2.32709487755 +    12721.57209941700*t);
+   earth_x_3+=     0.00000000011 * cos( 6.27313981891 +    12592.45001978260*t);
+   earth_x_3+=     0.00000000006 * cos( 2.56028859854 +      242.72860397400*t);
+   earth_x_3=earth_x_3 * t * t * t;
+
+   double earth_x_4=0.0;
+   earth_x_4+=     0.00000004022 * cos( 0.00000000000 +        0.00000000000*t);
+   earth_x_4+=     0.00000001927 * cos( 1.18746233453 +    12566.15169998280*t);
+   earth_x_4+=     0.00000000491 * cos( 5.29890244491 +     6283.07584999140*t);
+   earth_x_4+=     0.00000000210 * cos( 3.74319530048 +     6438.49624942560*t);
+   earth_x_4+=     0.00000000210 * cos( 2.90533827217 +     6127.65545055720*t);
+   earth_x_4+=     0.00000000229 * cos( 0.88072621187 +    18849.22754997420*t);
+   earth_x_4+=     0.00000000044 * cos( 4.59638160457 +      709.93304855830*t);
+   earth_x_4+=     0.00000000021 * cos( 0.18394116301 +     6279.55273164240*t);
+   earth_x_4+=     0.00000000021 * cos( 0.18140710246 +     6286.59896834040*t);
+   earth_x_4+=     0.00000000013 * cos( 3.81035106449 +    83996.84731811189*t);
+   earth_x_4+=     0.00000000011 * cos( 0.65552931190 +    25132.30339996560*t);
+   earth_x_4+=     0.00000000009 * cos( 5.19374498802 +    11856.21865142450*t);
+   earth_x_4+=     0.00000000007 * cos( 6.16589776486 +   167283.76158766549*t);
+   earth_x_4+=     0.00000000005 * cos( 2.68915538708 +    12410.73130054860*t);
+   earth_x_4+=     0.00000000006 * cos( 5.56853736836 +   156137.47598479928*t);
+   earth_x_4=earth_x_4 * t * t * t * t;
+
+   double earth_x_5=0.0;
+   earth_x_5+=     0.00000000108 * cos( 0.00000000000 +        0.00000000000*t);
+   earth_x_5+=     0.00000000043 * cos( 6.08789652058 +    12566.15169998280*t);
+   earth_x_5+=     0.00000000025 * cos( 5.33779665346 +     6438.49624942560*t);
+   earth_x_5+=     0.00000000025 * cos( 1.31073691919 +     6127.65545055720*t);
+   earth_x_5+=     0.00000000020 * cos( 0.94333339828 +     6283.07584999140*t);
+   earth_x_5+=     0.00000000008 * cos( 5.52588061142 +    18849.22754997420*t);
+   earth_x_5=earth_x_5 * t * t * t * t * t;
+
+   return earth_x_0+earth_x_1+earth_x_2+earth_x_3+earth_x_4+earth_x_5;
+}
+
+double vsop87a_full::earth_y(double t){
+   double earth_y_0=0.0;
+   earth_y_0+=     0.99989211030 * cos( 0.18265890456 +     6283.07584999140*t);
+   earth_y_0+=     0.02442699036 * cos( 3.14159265359 +        0.00000000000*t);
+   earth_y_0+=     0.00835292314 * cos( 0.13952878991 +    12566.15169998280*t);
+   earth_y_0+=     0.00010466965 * cos( 0.09641690558 +    18849.22754997420*t);
+   earth_y_0+=     0.00003110838 * cos( 5.38114091484 +    83996.84731811189*t);
+   earth_y_0+=     0.00002570338 * cos( 5.30103973360 +      529.69096509460*t);
+   earth_y_0+=     0.00002147473 * cos( 2.66253538905 +     1577.34354244780*t);
+   earth_y_0+=     0.00001709219 * cos( 5.20780401071 +     6279.55273164240*t);
+   earth_y_0+=     0.00001707987 * cos( 4.58232858766 +     6286.59896834040*t);
+   earth_y_0+=     0.00001440265 * cos( 1.90068164664 +     2352.86615377180*t);
+   earth_y_0+=     0.00001135092 * cos( 5.27313415220 +     5223.69391980220*t);
+   earth_y_0+=     0.00000934539 * cos( 4.50301201844 +    12036.46073488820*t);
+   earth_y_0+=     0.00000900565 * cos( 1.60563288120 +    10213.28554621100*t);
+   earth_y_0+=     0.00000567126 * cos( 0.58142248753 +     1059.38193018920*t);
+   earth_y_0+=     0.00000744932 * cos( 2.80728871886 +      398.14900340820*t);
+   earth_y_0+=     0.00000639316 * cos( 6.02923915017 +     5753.38488489680*t);
+   earth_y_0+=     0.00000681324 * cos( 0.64729627497 +     4705.73230754360*t);
+   earth_y_0+=     0.00000611347 * cos( 3.81381495286 +     6812.76681508600*t);
+   earth_y_0+=     0.00000450435 * cos( 4.52785572489 +     5884.92684658320*t);
+   earth_y_0+=     0.00000452018 * cos( 5.99167242707 +     6256.77753019160*t);
+   earth_y_0+=     0.00000449968 * cos( 3.79880375595 +     6309.37416979120*t);
+   earth_y_0+=     0.00000551390 * cos( 3.96125249369 +     5507.55323866740*t);
+   earth_y_0+=     0.00000406334 * cos( 5.25616268027 +     6681.22485339960*t);
+   earth_y_0+=     0.00000541273 * cos( 5.49902805917 +      775.52261132400*t);
+   earth_y_0+=     0.00000546360 * cos( 6.17311131785 +    14143.49524243060*t);
+   earth_y_0+=     0.00000507084 * cos( 2.87025193381 +     7860.41939243920*t);
+   earth_y_0+=     0.00000219504 * cos( 2.95216139568 +    11506.76976979360*t);
+   earth_y_0+=     0.00000227937 * cos( 5.95179248814 +     7058.59846131540*t);
+   earth_y_0+=     0.00000227792 * cos( 4.84547074733 +     4694.00295470760*t);
+   earth_y_0+=     0.00000255845 * cos( 0.69454231563 +    12168.00269657460*t);
+   earth_y_0+=     0.00000256132 * cos( 6.16722512388 +      709.93304855830*t);
+   earth_y_0+=     0.00000179242 * cos( 1.40003446021 +      796.29800681640*t);
+   earth_y_0+=     0.00000178280 * cos( 5.11717552231 +     6283.00853968860*t);
+   earth_y_0+=     0.00000178280 * cos( 4.67307255246 +     6283.14316029419*t);
+   earth_y_0+=     0.00000155454 * cos( 0.05340525434 +    25132.30339996560*t);
+   earth_y_0+=     0.00000206257 * cos( 4.28366728882 +    11790.62908865880*t);
+   earth_y_0+=     0.00000149769 * cos( 6.07429023278 +     5486.77784317500*t);
+   earth_y_0+=     0.00000200005 * cos( 2.50144088120 +    17789.84561978500*t);
+   earth_y_0+=     0.00000129006 * cos( 3.64623708634 +     7079.37385680780*t);
+   earth_y_0+=     0.00000128211 * cos( 3.23254821381 +     3738.76143010800*t);
+   earth_y_0+=     0.00000152790 * cos( 5.58120800450 +      213.29909543800*t);
+   earth_y_0+=     0.00000118725 * cos( 5.45361490488 +     9437.76293488700*t);
+   earth_y_0+=     0.00000078419 * cos( 0.28253093914 +     8827.39026987480*t);
+   earth_y_0+=     0.00000073513 * cos( 1.12433371880 +     1589.07289528380*t);
+   earth_y_0+=     0.00000065694 * cos( 2.08898587059 +    11769.85369316640*t);
+   earth_y_0+=     0.00000063240 * cos( 0.67044491570 +     6262.30045449900*t);
+   earth_y_0+=     0.00000062961 * cos( 2.83649755996 +     6303.85124548380*t);
+   earth_y_0+=     0.00000085870 * cos( 1.45351152151 +   167283.76158766549*t);
+   earth_y_0+=     0.00000063616 * cos( 0.49038985445 +        7.11354700080*t);
+   earth_y_0+=     0.00000080740 * cos( 0.92225805400 +     6069.77675455340*t);
+   earth_y_0+=     0.00000083511 * cos( 4.63223527298 +     3340.61242669980*t);
+   earth_y_0+=     0.00000060039 * cos( 1.82878888713 +     4136.91043351620*t);
+   earth_y_0+=     0.00000075626 * cos( 2.58860726122 +     6496.37494542940*t);
+   earth_y_0+=     0.00000058089 * cos( 0.55165509745 +    10977.07880469900*t);
+   earth_y_0+=     0.00000053747 * cos( 6.27553948416 +     1194.44701022460*t);
+   earth_y_0+=     0.00000057574 * cos( 1.10802134270 +     6282.09552892320*t);
+   earth_y_0+=     0.00000057574 * cos( 2.39891957636 +     6284.05617105960*t);
+   earth_y_0+=     0.00000051141 * cos( 6.02458046766 +     6290.18939699220*t);
+   earth_y_0+=     0.00000050679 * cos( 3.77133779892 +     6275.96230299060*t);
+   earth_y_0+=     0.00000049613 * cos( 1.07290390127 +     6127.65545055720*t);
+   earth_y_0+=     0.00000047934 * cos( 2.43176952238 +     6438.49624942560*t);
+   earth_y_0+=     0.00000061133 * cos( 5.34335721244 +     3154.68708489560*t);
+   earth_y_0+=     0.00000048425 * cos( 0.05833183710 +      801.82093112380*t);
+   earth_y_0+=     0.00000047359 * cos( 4.53607945311 +     3128.38876509580*t);
+   earth_y_0+=     0.00000048379 * cos( 0.25856696781 +     7084.89678111520*t);
+   earth_y_0+=     0.00000051293 * cos( 4.21367834558 +    12352.85260454480*t);
+   earth_y_0+=     0.00000050630 * cos( 0.48144269678 +    11856.21865142450*t);
+   earth_y_0+=     0.00000037539 * cos( 1.60040603609 +     8429.24126646660*t);
+   earth_y_0+=     0.00000049347 * cos( 2.15740113992 +     2544.31441988340*t);
+   earth_y_0+=     0.00000039555 * cos( 3.35943284176 +     5481.25491886760*t);
+   earth_y_0+=     0.00000048913 * cos( 5.16531049040 +     9225.53927328300*t);
+   earth_y_0+=     0.00000042135 * cos( 2.06532560040 +    17260.15465469040*t);
+   earth_y_0+=     0.00000037954 * cos( 1.47005640108 +      426.59819087600*t);
+   earth_y_0+=     0.00000031140 * cos( 3.61579568733 +     2146.16541647520*t);
+   earth_y_0+=     0.00000032603 * cos( 4.41806626166 +     4164.31198961300*t);
+   earth_y_0+=     0.00000039728 * cos( 4.70048863454 +     5856.47765911540*t);
+   earth_y_0+=     0.00000028515 * cos( 4.53853388964 +    12569.67481833180*t);
+   earth_y_0+=     0.00000036368 * cos( 1.33416490018 +    18073.70493865020*t);
+   earth_y_0+=     0.00000025032 * cos( 3.74982021361 +       26.29831979980*t);
+   earth_y_0+=     0.00000030329 * cos( 0.42324614391 +     4535.05943692440*t);
+   earth_y_0+=     0.00000026741 * cos( 5.97172882162 +     3930.20969621960*t);
+   earth_y_0+=     0.00000028166 * cos( 5.19366222912 +    12562.62858163380*t);
+   earth_y_0+=     0.00000024843 * cos( 1.00265326013 +     5088.62883976680*t);
+   earth_y_0+=     0.00000021642 * cos( 3.06444526478 +    12559.03815298200*t);
+   earth_y_0+=     0.00000025701 * cos( 0.75895390818 +    13367.97263110660*t);
+   earth_y_0+=     0.00000020498 * cos( 3.44435384536 +    11371.70468975820*t);
+   earth_y_0+=     0.00000025997 * cos( 1.87312778491 +    71430.69561812909*t);
+   earth_y_0+=     0.00000017888 * cos( 5.01504194498 +     1748.01641306700*t);
+   earth_y_0+=     0.00000023246 * cos( 4.61910038205 +    18319.53658487960*t);
+   earth_y_0+=     0.00000017077 * cos( 5.29722216277 +     4933.20844033260*t);
+   earth_y_0+=     0.00000021070 * cos( 2.22959477131 +     7477.52286021600*t);
+   earth_y_0+=     0.00000017512 * cos( 6.16192589291 +    10447.38783960440*t);
+   earth_y_0+=     0.00000020024 * cos( 2.96823180998 +     8031.09226305840*t);
+   earth_y_0+=     0.00000013831 * cos( 3.02450528624 +     8635.94200376320*t);
+   earth_y_0+=     0.00000014095 * cos( 6.03707103538 +     2942.46342329160*t);
+   earth_y_0+=     0.00000018884 * cos( 0.85614789376 +   156137.47598479928*t);
+   earth_y_0+=     0.00000017861 * cos( 4.86809570283 +     1592.59601363280*t);
+   earth_y_0+=     0.00000013605 * cos( 3.62641145030 +    13095.84266507740*t);
+   earth_y_0+=     0.00000015078 * cos( 1.55132224732 +    16496.36139620240*t);
+   earth_y_0+=     0.00000017985 * cos( 2.07741211308 +    17298.18232732620*t);
+   earth_y_0+=     0.00000014386 * cos( 6.14898732933 +    20426.57109242200*t);
+   earth_y_0+=     0.00000012227 * cos( 4.38359423735 +     7632.94325965020*t);
+   earth_y_0+=     0.00000013376 * cos( 2.03616887071 +    12139.55350910680*t);
+   earth_y_0+=     0.00000011583 * cos( 3.88971333608 +     5331.35744374080*t);
+   earth_y_0+=     0.00000013558 * cos( 1.44666331503 +    16730.46368959580*t);
+   earth_y_0+=     0.00000014332 * cos( 4.45678693729 +     7342.45778018060*t);
+   earth_y_0+=     0.00000012375 * cos( 1.53862289477 +     6279.48542133960*t);
+   earth_y_0+=     0.00000012366 * cos( 1.96831814872 +     6286.66627864320*t);
+   earth_y_0+=     0.00000012446 * cos( 5.78094918529 +    15110.46611986620*t);
+   earth_y_0+=     0.00000010989 * cos( 0.62040343662 +    10988.80815753500*t);
+   earth_y_0+=     0.00000012556 * cos( 2.59194808538 +    15720.83878487840*t);
+   earth_y_0+=     0.00000011583 * cos( 2.48165809843 +     5729.50644714900*t);
+   earth_y_0+=     0.00000009023 * cos( 3.85655029499 +     9623.68827669120*t);
+   earth_y_0+=     0.00000009121 * cos( 5.83006047147 +     7234.79425624200*t);
+   earth_y_0+=     0.00000008869 * cos( 3.47611265290 +     6148.01076995600*t);
+   earth_y_0+=     0.00000008614 * cos( 0.03081589303 +     6418.14093002680*t);
+   earth_y_0+=     0.00000007911 * cos( 1.58348238359 +     2118.76386037840*t);
+   earth_y_0+=     0.00000008293 * cos( 0.85407021371 +    14712.31711645800*t);
+   earth_y_0+=     0.00000010219 * cos( 0.13075179540 +     1349.86740965880*t);
+   earth_y_0+=     0.00000010126 * cos( 0.99630008630 +     6836.64525283380*t);
+   earth_y_0+=     0.00000010480 * cos( 1.38209301929 +     5999.21653112620*t);
+   earth_y_0+=     0.00000007353 * cos( 3.83220307326 +     6040.34724601740*t);
+   earth_y_0+=     0.00000009820 * cos( 2.12476655144 +     6566.93516885660*t);
+   earth_y_0+=     0.00000007018 * cos( 5.94659180519 +     6525.80445396540*t);
+   earth_y_0+=     0.00000006874 * cos( 4.57503512461 +     7238.67559160000*t);
+   earth_y_0+=     0.00000007545 * cos( 5.94843444430 +    12539.85338018300*t);
+   earth_y_0+=     0.00000009540 * cos( 3.49613271706 +    21228.39202354580*t);
+   earth_y_0+=     0.00000007202 * cos( 4.31311898936 +     6245.04817735560*t);
+   earth_y_0+=     0.00000007144 * cos( 5.47703560836 +     6321.10352262720*t);
+   earth_y_0+=     0.00000006627 * cos( 5.41373102778 +     5327.47610838280*t);
+   earth_y_0+=     0.00000008544 * cos( 1.82829301857 +    11015.10647733480*t);
+   earth_y_0+=     0.00000008602 * cos( 5.43953575310 +     5643.17856367740*t);
+   earth_y_0+=     0.00000006158 * cos( 5.65852675233 +    12964.30070339100*t);
+   earth_y_0+=     0.00000007414 * cos( 3.46720433621 +    23543.23050468179*t);
+   earth_y_0+=     0.00000007317 * cos( 2.01664477821 +      536.80451209540*t);
+   earth_y_0+=     0.00000006352 * cos( 3.46461304398 +     1990.74501704100*t);
+   earth_y_0+=     0.00000006251 * cos( 1.52860898333 +      951.71840625060*t);
+   earth_y_0+=     0.00000005828 * cos( 1.37555383064 +      955.59974160860*t);
+   earth_y_0+=     0.00000006389 * cos( 2.40580071461 +    24072.92146977640*t);
+   earth_y_0+=     0.00000006976 * cos( 4.76626361488 +    10973.55568635000*t);
+   earth_y_0+=     0.00000006916 * cos( 5.39658071678 +      522.57741809380*t);
+   earth_y_0+=     0.00000007459 * cos( 2.71133898870 +    22003.91463486980*t);
+   earth_y_0+=     0.00000005376 * cos( 2.34968505500 +    14314.16811304980*t);
+   earth_y_0+=     0.00000006852 * cos( 4.33781274406 +      155.42039943420*t);
+   earth_y_0+=     0.00000007527 * cos( 1.94392127074 +    18422.62935909819*t);
+   earth_y_0+=     0.00000006222 * cos( 0.67296135144 +    18451.07854656599*t);
+   earth_y_0+=     0.00000005836 * cos( 0.75921238848 +      639.89728631400*t);
+   earth_y_0+=     0.00000006487 * cos( 2.22135874272 +     5216.58037280140*t);
+   earth_y_0+=     0.00000006251 * cos( 2.88399458161 +     6288.59877429880*t);
+   earth_y_0+=     0.00000006251 * cos( 0.62294633745 +     6277.55292568400*t);
+   earth_y_0+=     0.00000007004 * cos( 3.57492830242 +      103.09277421860*t);
+   earth_y_0+=     0.00000005337 * cos( 1.97127813671 +     4690.47983635860*t);
+   earth_y_0+=     0.00000005605 * cos( 1.92150700481 +     1551.04522264800*t);
+   earth_y_0+=     0.00000006210 * cos( 0.81866451118 +     5650.29211067820*t);
+   earth_y_0+=     0.00000006064 * cos( 5.19708779870 +     5230.80746680300*t);
+   earth_y_0+=     0.00000004596 * cos( 2.92936593195 +      553.56940284240*t);
+   earth_y_0+=     0.00000005386 * cos( 5.40633950052 +    77713.77146812050*t);
+   earth_y_0+=     0.00000006332 * cos( 2.83725098708 +    26087.90314157420*t);
+   earth_y_0+=     0.00000006196 * cos( 4.15787078176 +    84672.47584450469*t);
+   earth_y_0+=     0.00000006160 * cos( 3.21170435536 +   239424.39025435288*t);
+   earth_y_0+=     0.00000005016 * cos( 1.46006858313 +     4732.03062734340*t);
+   earth_y_0+=     0.00000005414 * cos( 1.42882325200 +     6179.98307577280*t);
+   earth_y_0+=     0.00000005413 * cos( 2.08325983447 +     6386.16862421000*t);
+   earth_y_0+=     0.00000004863 * cos( 2.19642746549 +    90279.92316810328*t);
+   earth_y_0+=     0.00000005711 * cos( 2.75630374926 +     6915.85958930460*t);
+   earth_y_0+=     0.00000004270 * cos( 0.81269903423 +     7875.67186362420*t);
+   earth_y_0+=     0.00000004017 * cos( 2.21501265909 +     6254.62666252360*t);
+   earth_y_0+=     0.00000004071 * cos( 3.99159158863 +     3634.62102451840*t);
+   earth_y_0+=     0.00000003993 * cos( 1.29192500625 +     6311.52503745920*t);
+   earth_y_0+=     0.00000004858 * cos( 1.27914374675 +     5760.49843189760*t);
+   earth_y_0+=     0.00000005274 * cos( 4.91687036691 +    25158.60171976540*t);
+   earth_y_0+=     0.00000004120 * cos( 1.72568521475 +     6709.67404086740*t);
+   earth_y_0+=     0.00000004410 * cos( 2.23536264210 +     6805.65326808520*t);
+   earth_y_0+=     0.00000003174 * cos( 3.62030131501 +     1066.49547719000*t);
+   earth_y_0+=     0.00000003064 * cos( 2.46605561478 +     5863.59120611620*t);
+   earth_y_0+=     0.00000003083 * cos( 5.68731909779 +     9917.69687450980*t);
+   earth_y_0+=     0.00000003385 * cos( 3.81076637101 +    13916.01910964160*t);
+   earth_y_0+=     0.00000002979 * cos( 4.62987940723 +    12566.21901028560*t);
+   earth_y_0+=     0.00000003980 * cos( 5.90917210624 +    13341.67431130680*t);
+   earth_y_0+=     0.00000003290 * cos( 0.95752698537 +    16200.77272450120*t);
+   earth_y_0+=     0.00000003885 * cos( 3.94360149823 +       74.78159856730*t);
+   earth_y_0+=     0.00000003060 * cos( 1.63142398497 +    18052.92954315780*t);
+   earth_y_0+=     0.00000003745 * cos( 3.65023712104 +     5966.68398033480*t);
+   earth_y_0+=     0.00000003592 * cos( 4.03843804285 +     6208.29425142410*t);
+   earth_y_0+=     0.00000003581 * cos( 5.75234129258 +     6357.85744855870*t);
+   earth_y_0+=     0.00000003626 * cos( 6.13871061291 +     6599.46771964800*t);
+   earth_y_0+=     0.00000002927 * cos( 1.79201822142 +    12721.57209941700*t);
+   earth_y_0+=     0.00000003625 * cos( 1.43112641053 +     6019.99192661860*t);
+   earth_y_0+=     0.00000002589 * cos( 2.27559640821 +    12592.45001978260*t);
+   earth_y_0+=     0.00000002558 * cos( 6.08151329647 +    10575.40668294180*t);
+   earth_y_0+=     0.00000002670 * cos( 2.64370887209 +        3.52311834900*t);
+   earth_y_0+=     0.00000002611 * cos( 0.86180315523 +     6702.56049386660*t);
+   earth_y_0+=     0.00000002535 * cos( 0.01231078888 +    31415.37924995700*t);
+   earth_y_0+=     0.00000003520 * cos( 3.80906746266 +   250570.67585721909*t);
+   earth_y_0+=     0.00000003434 * cos( 2.07577333297 +     6546.15977336420*t);
+   earth_y_0+=     0.00000002378 * cos( 2.05651519036 +     2388.89402044920*t);
+   earth_y_0+=     0.00000002979 * cos( 5.07398237708 +    12566.08438968000*t);
+   earth_y_0+=     0.00000002716 * cos( 1.49855840728 +    12029.34718788740*t);
+   earth_y_0+=     0.00000002215 * cos( 3.69574321309 +    13362.44970679920*t);
+   earth_y_0+=     0.00000003049 * cos( 0.05228873308 +    29088.81141598500*t);
+   earth_y_0+=     0.00000002784 * cos( 3.37954165171 +    14945.31617355440*t);
+   earth_y_0+=     0.00000002284 * cos( 4.38146814547 +     6076.89030155420*t);
+   earth_y_0+=     0.00000002529 * cos( 4.37629589914 +    12043.57428188900*t);
+   earth_y_0+=     0.00000002015 * cos( 0.91352603582 +     6262.72053059260*t);
+   earth_y_0+=     0.00000002218 * cos( 5.40562909521 +     6489.26139842860*t);
+   earth_y_0+=     0.00000002005 * cos( 2.59344263770 +     6303.43116939020*t);
+   earth_y_0+=     0.00000002681 * cos( 5.85620039558 +     4590.91018048900*t);
+   earth_y_0+=     0.00000002011 * cos( 5.69828891426 +     4701.11650170840*t);
+   earth_y_0+=     0.00000002440 * cos( 4.55844746985 +      149.56319713460*t);
+   earth_y_0+=     0.00000002600 * cos( 3.93109819647 +    19651.04848109800*t);
+   earth_y_0+=     0.00000002197 * cos( 5.22411628121 +    13517.87010623340*t);
+   earth_y_0+=     0.00000001893 * cos( 4.09716259686 +     6286.95718534940*t);
+   earth_y_0+=     0.00000001893 * cos( 5.69296362938 +     6279.19451463340*t);
+   earth_y_0+=     0.00000002402 * cos( 4.84503253270 +     1162.47470440780*t);
+   earth_y_0+=     0.00000002156 * cos( 0.43826128691 +    18139.29450141590*t);
+   earth_y_0+=     0.00000001778 * cos( 2.79119893202 +     4292.33083295040*t);
+   earth_y_0+=     0.00000001966 * cos( 1.82224928185 +     4686.88940770680*t);
+   earth_y_0+=     0.00000002327 * cos( 2.83472726504 +    10021.83728009940*t);
+   earth_y_0+=     0.00000002436 * cos( 2.83689873280 +    95143.13292097810*t);
+   earth_y_0+=     0.00000002163 * cos( 6.13036765564 +     6037.24420376200*t);
+   earth_y_0+=     0.00000001841 * cos( 3.67947373190 +      242.72860397400*t);
+   earth_y_0+=     0.00000001672 * cos( 0.33097080655 +    13119.72110282519*t);
+   earth_y_0+=     0.00000002250 * cos( 0.82649282051 +    12669.24447420140*t);
+   earth_y_0+=     0.00000001627 * cos( 4.77445737868 +    13521.75144159140*t);
+   earth_y_0+=     0.00000001888 * cos( 2.87850805387 +     8662.24032356300*t);
+   earth_y_0+=     0.00000002052 * cos( 3.65927016475 +     6528.90749622080*t);
+   earth_y_0+=     0.00000001774 * cos( 0.51400448882 +     6819.88036208680*t);
+   earth_y_0+=     0.00000001840 * cos( 1.91711006767 +    65147.61976813770*t);
+   earth_y_0+=     0.00000001544 * cos( 1.33911061474 +      110.20632121940*t);
+   earth_y_0+=     0.00000001756 * cos( 2.96482886554 +     5746.27133789600*t);
+   earth_y_0+=     0.00000001543 * cos( 4.39018655159 +     6272.43918464160*t);
+   earth_y_0+=     0.00000001543 * cos( 5.39993967465 +     6293.71251534120*t);
+   earth_y_0+=     0.00000001952 * cos( 3.09643550846 +    23013.53953958720*t);
+   earth_y_0+=     0.00000001484 * cos( 2.15301976837 +       20.77539549240*t);
+   earth_y_0+=     0.00000001558 * cos( 2.79679951671 +    17654.78053974960*t);
+   earth_y_0+=     0.00000001440 * cos( 0.69014829975 +     1052.26838318840*t);
+   earth_y_0+=     0.00000001648 * cos( 3.98135416194 +      632.78373931320*t);
+   earth_y_0+=     0.00000001795 * cos( 1.47161886169 +    33019.02111220460*t);
+   earth_y_0+=     0.00000001488 * cos( 4.10377440731 +       76.26607127560*t);
+   earth_y_0+=     0.00000001349 * cos( 4.18018072818 +     6016.46880826960*t);
+   earth_y_0+=     0.00000001477 * cos( 2.56925128481 +    12779.45079542080*t);
+   earth_y_0+=     0.00000001252 * cos( 6.26830198321 +    11499.65622279280*t);
+   earth_y_0+=     0.00000001491 * cos( 4.65893645946 +     9411.46461508720*t);
+   earth_y_0+=     0.00000001471 * cos( 4.18695456530 +    18635.92845453620*t);
+   earth_y_0+=     0.00000001263 * cos( 5.62967869388 +    12573.26524698360*t);
+   earth_y_0+=     0.00000001385 * cos( 2.65776727703 +     2379.16447357160*t);
+   earth_y_0+=     0.00000001272 * cos( 5.60977857361 +     6549.68289171320*t);
+   earth_y_0+=     0.00000001578 * cos( 4.10046385407 +    25934.12433108940*t);
+   earth_y_0+=     0.00000001557 * cos( 4.78614835776 +    18875.52586977400*t);
+   earth_y_0+=     0.00000001358 * cos( 0.33781842226 +     5120.60114558360*t);
+   earth_y_0+=     0.00000001505 * cos( 2.70837129901 +    18209.33026366019*t);
+   earth_y_0+=     0.00000001460 * cos( 3.40156973977 +      245.83164622940*t);
+   earth_y_0+=     0.00000001252 * cos( 6.12407876457 +     6993.00889854970*t);
+   earth_y_0+=     0.00000001242 * cos( 1.27359177098 +    24356.78078864160*t);
+   earth_y_0+=     0.00000001464 * cos( 3.73000165992 +       38.13303563780*t);
+   earth_y_0+=     0.00000001028 * cos( 0.93950338023 +     6112.40297937220*t);
+   earth_y_0+=     0.00000001014 * cos( 4.36072926630 +    17256.63153634140*t);
+   earth_y_0+=     0.00000001312 * cos( 5.81330496460 +    12491.37010141550*t);
+   earth_y_0+=     0.00000000992 * cos( 2.56738139157 +     6453.74872061060*t);
+   earth_y_0+=     0.00000000988 * cos( 5.69171453384 +     8273.82086703240*t);
+   earth_y_0+=     0.00000001191 * cos( 1.64765684664 +     6276.02961329339*t);
+   earth_y_0+=     0.00000001191 * cos( 1.85928407242 +     6290.12208668940*t);
+   earth_y_0+=     0.00000000978 * cos( 3.16773119822 +     6247.04798331399*t);
+   earth_y_0+=     0.00000001223 * cos( 6.22694723537 +     5429.87946823940*t);
+   earth_y_0+=     0.00000000971 * cos( 0.33921097188 +     6319.10371666880*t);
+   earth_y_0+=     0.00000000962 * cos( 1.06488912173 +    12565.17137891460*t);
+   earth_y_0+=     0.00000001025 * cos( 0.64688768472 +    12545.37630449040*t);
+   earth_y_0+=     0.00000001234 * cos( 2.50427756599 +    22483.84857449259*t);
+   earth_y_0+=     0.00000001134 * cos( 5.76287831677 +    17782.73207278420*t);
+   earth_y_0+=     0.00000000944 * cos( 5.49712708636 +     6306.95428773920*t);
+   earth_y_0+=     0.00000000944 * cos( 4.29299913988 +     6259.19741224360*t);
+   earth_y_0+=     0.00000001092 * cos( 2.35048657695 +    17796.95916678580*t);
+   earth_y_0+=     0.00000000995 * cos( 1.11794413013 +    10177.25767953360*t);
+   earth_y_0+=     0.00000000894 * cos( 2.94681683620 +      206.18554843720*t);
+   earth_y_0+=     0.00000001216 * cos( 5.61396859194 +    82576.98122099529*t);
+   earth_y_0+=     0.00000000860 * cos( 6.19327755003 +     1692.16566950240*t);
+   earth_y_0+=     0.00000000942 * cos( 6.13301518193 +     6267.82337880639*t);
+   earth_y_0+=     0.00000000942 * cos( 3.65711104430 +     6298.32832117640*t);
+   earth_y_0+=     0.00000001045 * cos( 5.17292899506 +    15508.61512327440*t);
+   earth_y_0+=     0.00000000850 * cos( 6.16443178844 +     6393.28217121080*t);
+   earth_y_0+=     0.00000000873 * cos( 2.86766306801 +      949.17560896980*t);
+   earth_y_0+=     0.00000000896 * cos( 0.83013691514 +     3903.91137641980*t);
+   earth_y_0+=     0.00000001023 * cos( 5.42813246189 +      220.41264243880*t);
+   earth_y_0+=     0.00000000856 * cos( 1.36181993943 +     4804.20927592700*t);
+   earth_y_0+=     0.00000001033 * cos( 3.31340063953 +     5573.14280143310*t);
+   earth_y_0+=     0.00000000817 * cos( 3.58855573226 +     6172.86952877200*t);
+   earth_y_0+=     0.00000000917 * cos( 0.89919993710 +   149854.40013480789*t);
+   earth_y_0+=     0.00000001081 * cos( 2.89245979526 +    36949.23080842420*t);
+   earth_y_0+=     0.00000000779 * cos( 2.01260281552 +     2648.45482547300*t);
+   earth_y_0+=     0.00000000962 * cos( 2.35578735539 +    12567.13202105100*t);
+   earth_y_0+=     0.00000000990 * cos( 4.97210778261 +    11403.67699557500*t);
+   earth_y_0+=     0.00000000739 * cos( 3.39896198161 +    11926.25441366880*t);
+   earth_y_0+=     0.00000000929 * cos( 0.64994792213 +     2787.04302385740*t);
+   earth_y_0+=     0.00000000973 * cos( 6.19358824421 +    22805.73556599360*t);
+   earth_y_0+=     0.00000000743 * cos( 2.22828430376 +     7834.12107263940*t);
+   earth_y_0+=     0.00000000905 * cos( 2.24459183660 +     6284.56032269970*t);
+   earth_y_0+=     0.00000000905 * cos( 1.26234908246 +     6281.59137728310*t);
+   earth_y_0+=     0.00000000798 * cos( 1.05421393588 +    12410.73130054860*t);
+   earth_y_0+=     0.00000000813 * cos( 5.23283372183 +    12416.58850284820*t);
+   earth_y_0+=     0.00000000869 * cos( 1.74785663537 +     6133.51265285680*t);
+   earth_y_0+=     0.00000000756 * cos( 5.40063116731 +     4061.21921539440*t);
+   earth_y_0+=     0.00000000616 * cos( 4.13389755770 +     1596.18644228460*t);
+   earth_y_0+=     0.00000000724 * cos( 1.49651873854 +   161000.68573767410*t);
+   earth_y_0+=     0.00000000612 * cos( 5.89184406078 +    16858.48253293320*t);
+   earth_y_0+=     0.00000000814 * cos( 3.82004246488 +     3894.18182954220*t);
+   earth_y_0+=     0.00000000785 * cos( 2.44918664136 +     4907.30205014560*t);
+   earth_y_0+=     0.00000000788 * cos( 1.86746379419 +      846.08283475120*t);
+   earth_y_0+=     0.00000000611 * cos( 6.24387943735 +      853.19638175200*t);
+   earth_y_0+=     0.00000000570 * cos( 2.97586328906 +    12586.92709547520*t);
+   earth_y_0+=     0.00000000561 * cos( 5.85814246596 +    12323.42309600880*t);
+   earth_y_0+=     0.00000000646 * cos( 5.91116942848 +    52175.80628314840*t);
+   earth_y_0+=     0.00000000590 * cos( 0.47736597307 +    15671.08175940660*t);
+   earth_y_0+=     0.00000000555 * cos( 2.86698301793 +    11513.88331679440*t);
+   earth_y_0+=     0.00000000536 * cos( 4.49580958272 +    18852.75066832320*t);
+   earth_y_0+=     0.00000000737 * cos( 5.41358302312 +    23581.25817731760*t);
+   earth_y_0+=     0.00000000532 * cos( 3.01682571392 +    18842.11400297339*t);
+   earth_y_0+=     0.00000000552 * cos( 6.08260944683 +    11933.36796066960*t);
+   earth_y_0+=     0.00000000599 * cos( 0.54713576844 +     6151.53388830500*t);
+   earth_y_0+=     0.00000000540 * cos( 1.83361074884 +     6279.78949257360*t);
+   earth_y_0+=     0.00000000540 * cos( 1.67333017022 +     6286.36220740920*t);
+   earth_y_0+=     0.00000000513 * cos( 2.13211972766 +    12455.94537876340*t);
+   earth_y_0+=     0.00000000649 * cos( 3.33488856037 +    11764.33076885900*t);
+   earth_y_0+=     0.00000000582 * cos( 2.96012819576 +     6414.61781167780*t);
+   earth_y_0+=     0.00000000608 * cos( 2.37588807672 +    10419.98628350760*t);
+   earth_y_0+=     0.00000000620 * cos( 0.83937590497 +     5326.78669402080*t);
+   earth_y_0+=     0.00000000536 * cos( 2.13145334966 +     7232.25145896120*t);
+   earth_y_0+=     0.00000000538 * cos( 1.33481115516 +     5333.90024102160*t);
+   earth_y_0+=     0.00000000516 * cos( 1.35797709246 +      433.71173787680*t);
+   earth_y_0+=     0.00000000596 * cos( 0.99065349155 +     6062.66320755260*t);
+   earth_y_0+=     0.00000000664 * cos( 4.31358012484 +    40879.44050464380*t);
+   earth_y_0+=     0.00000000607 * cos( 4.79849916312 +     5849.36411211460*t);
+   earth_y_0+=     0.00000000629 * cos( 1.34455816041 +    26735.94526221320*t);
+   earth_y_0+=     0.00000000479 * cos( 6.04755907663 +     5547.19933645960*t);
+   earth_y_0+=     0.00000000565 * cos( 5.49238261982 +     4171.42553661380*t);
+   earth_y_0+=     0.00000000472 * cos( 3.81181804730 +    10969.96525769820*t);
+   earth_y_0+=     0.00000000462 * cos( 3.53729128828 +     3104.93005942380*t);
+   earth_y_0+=     0.00000000551 * cos( 2.46794798020 +     6503.48849243020*t);
+   earth_y_0+=     0.00000000535 * cos( 4.55197627124 +   173566.83743765688*t);
+   earth_y_0+=     0.00000000450 * cos( 0.38613288324 +     4274.51831083240*t);
+   earth_y_0+=     0.00000000527 * cos( 4.67747327047 +    24602.61243487099*t);
+   earth_y_0+=     0.00000000502 * cos( 5.68646136646 +    12489.88562870720*t);
+   earth_y_0+=     0.00000000603 * cos( 5.56725954240 +   322711.30452390644*t);
+   earth_y_0+=     0.00000000432 * cos( 5.22277533085 +     9388.00590941520*t);
+   earth_y_0+=     0.00000000530 * cos( 5.14730101928 +    18845.70443162520*t);
+   earth_y_0+=     0.00000000431 * cos( 3.74066314880 +     7018.95236352320*t);
+   earth_y_0+=     0.00000000484 * cos( 0.31541423564 +       14.22709400160*t);
+   earth_y_0+=     0.00000000560 * cos( 0.14432555287 +     1581.95934828300*t);
+   earth_y_0+=     0.00000000584 * cos( 4.78380792490 +    72850.56171524568*t);
+   earth_y_0+=     0.00000000408 * cos( 4.98636538552 +     3496.03282613400*t);
+   earth_y_0+=     0.00000000569 * cos( 1.38164944312 +    97238.62754448749*t);
+   earth_y_0+=     0.00000000478 * cos( 4.58275550120 +    29296.61538957860*t);
+   earth_y_0+=     0.00000000389 * cos( 1.85967717576 +     6244.94281435360*t);
+   earth_y_0+=     0.00000000389 * cos( 1.64726374330 +     6321.20888562920*t);
+   earth_y_0+=     0.00000000388 * cos( 5.85202122263 +     6297.30294399300*t);
+   earth_y_0+=     0.00000000388 * cos( 3.93810500361 +     6268.84875598980*t);
+   earth_y_0+=     0.00000000515 * cos( 0.09995703640 +    12320.32005375340*t);
+   earth_y_0+=     0.00000000395 * cos( 1.39225625253 +     4157.19844261220*t);
+   earth_y_0+=     0.00000000492 * cos( 4.39757189321 +     1478.86657406440*t);
+   earth_y_0+=     0.00000000423 * cos( 2.45281376602 +     9779.10867612540*t);
+   earth_y_0+=     0.00000000439 * cos( 3.25487944876 +   233141.31440436149*t);
+   earth_y_0+=     0.00000000501 * cos( 3.08355690853 +    17157.06188047180*t);
+   earth_y_0+=     0.00000000416 * cos( 2.17399987974 +    13760.59871020740*t);
+   earth_y_0+=     0.00000000503 * cos( 4.02857067171 +      135.06508003540*t);
+   earth_y_0+=     0.00000000393 * cos( 6.27884951137 +     7872.14874527520*t);
+   earth_y_0+=     0.00000000365 * cos( 0.53940975175 +     6922.97313630540*t);
+   earth_y_0+=     0.00000000409 * cos( 6.14577564379 +    10984.19235169980*t);
+   earth_y_0+=     0.00000000371 * cos( 4.19153897517 +    13625.53363017200*t);
+   earth_y_0+=     0.00000000370 * cos( 2.79623025524 +    14919.01785375460*t);
+   earth_y_0+=     0.00000000499 * cos( 1.15482230725 +   161710.61878623239*t);
+   earth_y_0+=     0.00000000379 * cos( 5.52650430530 +     3185.19202726560*t);
+   earth_y_0+=     0.00000000420 * cos( 1.76851201552 +    22779.43724619380*t);
+   earth_y_0+=     0.00000000485 * cos( 6.06683096360 +    12528.01866434500*t);
+   earth_y_0+=     0.00000000358 * cos( 1.58658367058 +     6271.34649715540*t);
+   earth_y_0+=     0.00000000358 * cos( 1.92035724848 +     6294.80520282740*t);
+   earth_y_0+=     0.00000000354 * cos( 2.96580857786 +    17267.26820169119*t);
+   earth_y_0+=     0.00000000351 * cos( 0.64808322427 +    11616.97609101300*t);
+   earth_y_0+=     0.00000000388 * cos( 5.20354943335 +      956.28915597060*t);
+   earth_y_0+=     0.00000000409 * cos( 2.03946298626 +    21954.15760939799*t);
+   earth_y_0+=     0.00000000369 * cos( 3.36327405725 +     5436.99301524020*t);
+   earth_y_0+=     0.00000000421 * cos( 3.25654135729 +     7445.55055439920*t);
+   earth_y_0+=     0.00000000361 * cos( 5.53538077166 +    21393.54196985760*t);
+   earth_y_0+=     0.00000000336 * cos( 6.12569565789 +     6279.14369672830*t);
+   earth_y_0+=     0.00000000336 * cos( 3.66443056835 +     6287.00800325450*t);
+   earth_y_0+=     0.00000000450 * cos( 1.80230243402 +     1385.56157495110*t);
+   earth_y_0+=     0.00000000315 * cos( 3.47678894225 +    10873.98603048040*t);
+   earth_y_0+=     0.00000000363 * cos( 2.71990145670 +     7335.34423317980*t);
+   earth_y_0+=     0.00000000347 * cos( 5.27040068469 +    17253.04110768959*t);
+   earth_y_0+=     0.00000000349 * cos( 1.04077270827 +      316.39186965660*t);
+   earth_y_0+=     0.00000000395 * cos( 3.82348293170 +     5905.70224207560*t);
+   earth_y_0+=     0.00000000308 * cos( 2.54337606421 +    20995.39296644940*t);
+   earth_y_0+=     0.00000000319 * cos( 5.30613709521 +       38.02767263580*t);
+   earth_y_0+=     0.00000000304 * cos( 4.87224757874 +     5749.86176654780*t);
+   earth_y_0+=     0.00000000414 * cos( 5.73458482925 +    44809.65020086340*t);
+   earth_y_0+=     0.00000000298 * cos( 2.28443357697 +      640.87760738220*t);
+   earth_y_0+=     0.00000000351 * cos( 4.36315900786 +    12345.73905754400*t);
+   earth_y_0+=     0.00000000304 * cos( 3.36893995706 +    29826.30635467320*t);
+   earth_y_0+=     0.00000000395 * cos( 4.78955929514 +     4176.04134244900*t);
+   earth_y_0+=     0.00000000329 * cos( 5.87961134010 +        3.93215326310*t);
+   earth_y_0+=     0.00000000374 * cos( 5.15697532254 +     5642.19824260920*t);
+   earth_y_0+=     0.00000000320 * cos( 0.92923666608 +    10818.13528691580*t);
+   earth_y_0+=     0.00000000365 * cos( 5.96664322939 +     6660.44945790720*t);
+   earth_y_0+=     0.00000000339 * cos( 4.53749520696 +     6923.95345737360*t);
+   earth_y_0+=     0.00000000357 * cos( 1.12153755723 +    16460.33352952499*t);
+   earth_y_0+=     0.00000000386 * cos( 0.71941423743 +     2636.72547263700*t);
+   earth_y_0+=     0.00000000273 * cos( 5.79109432105 +    26709.64694241340*t);
+   earth_y_0+=     0.00000000265 * cos( 3.53930695872 +    19378.91851506880*t);
+   earth_y_0+=     0.00000000325 * cos( 0.92198442483 +    12359.96615154560*t);
+   earth_y_0+=     0.00000000278 * cos( 3.86335874809 +    20597.24396304120*t);
+   earth_y_0+=     0.00000000294 * cos( 2.59935512947 +    28286.99048486120*t);
+   earth_y_0+=     0.00000000269 * cos( 4.91776226798 +     6816.28993343500*t);
+   earth_y_0+=     0.00000000357 * cos( 2.74654149425 +    30666.15495843280*t);
+   earth_y_0+=     0.00000000344 * cos( 1.94480070409 +     6147.45052498140*t);
+   earth_y_0+=     0.00000000255 * cos( 3.37524772579 +     1375.77379984580*t);
+   earth_y_0+=     0.00000000308 * cos( 5.40069482744 +    29864.33402730900*t);
+   earth_y_0+=     0.00000000334 * cos( 1.56223513612 +     6418.70117500140*t);
+   earth_y_0+=     0.00000000316 * cos( 3.97125662279 +     6055.59971720240*t);
+   earth_y_0+=     0.00000000243 * cos( 2.53093250846 +     6310.47740608820*t);
+   earth_y_0+=     0.00000000243 * cos( 0.97600841060 +     6255.67429389460*t);
+   earth_y_0+=     0.00000000296 * cos( 4.79138384974 +     6129.29703950660*t);
+   earth_y_0+=     0.00000000301 * cos( 5.81886973641 +     6510.55198278040*t);
+   earth_y_0+=     0.00000000329 * cos( 5.86928788098 +      377.37360791580*t);
+   earth_y_0+=     0.00000000282 * cos( 3.64505853337 +      263.08392337280*t);
+   earth_y_0+=     0.00000000267 * cos( 1.84936252726 +    24705.70520908959*t);
+   earth_y_0+=     0.00000000285 * cos( 0.93414827012 +     5469.52556603160*t);
+   earth_y_0+=     0.00000000234 * cos( 0.33992064306 +     5017.50837136500*t);
+   earth_y_0+=     0.00000000286 * cos( 4.99886287753 +     6436.85466047620*t);
+   earth_y_0+=     0.00000000309 * cos( 0.57495592500 +    27511.46787353720*t);
+   earth_y_0+=     0.00000000319 * cos( 3.13977654412 +   169379.25621117488*t);
+   earth_y_0+=     0.00000000263 * cos( 6.04953863141 +    11712.95531823080*t);
+   earth_y_0+=     0.00000000285 * cos( 5.85104532566 +    16522.65971600220*t);
+   earth_y_0+=     0.00000000230 * cos( 5.72898960063 +     1265.56747862640*t);
+   earth_y_0+=     0.00000000235 * cos( 4.94143074133 +     4487.81740627040*t);
+   earth_y_0+=     0.00000000253 * cos( 2.08490293021 +     4377.61108505100*t);
+   earth_y_0+=     0.00000000312 * cos( 1.61150992635 +    11720.06886523160*t);
+   earth_y_0+=     0.00000000220 * cos( 6.00356459702 +    20199.09495963300*t);
+   earth_y_0+=     0.00000000214 * cos( 0.52396088270 +     2221.85663459700*t);
+   earth_y_0+=     0.00000000213 * cos( 2.22365237019 +    30355.99731976780*t);
+   earth_y_0+=     0.00000000259 * cos( 1.86047749915 +     6282.83203250789*t);
+   earth_y_0+=     0.00000000259 * cos( 1.64646341991 +     6283.31966747490*t);
+   earth_y_0+=     0.00000000243 * cos( 2.57267015777 +     7096.62613395120*t);
+   earth_y_0+=     0.00000000221 * cos( 4.65757923242 +     5540.08578945880*t);
+   earth_y_0+=     0.00000000246 * cos( 4.26495076165 +     8671.96987044060*t);
+   earth_y_0+=     0.00000000208 * cos( 3.10897679499 +    10440.27429260360*t);
+   earth_y_0+=     0.00000000199 * cos( 5.64807046697 +    11216.28429032400*t);
+   earth_y_0+=     0.00000000209 * cos( 3.71139099761 +      419.48464387520*t);
+   earth_y_0+=     0.00000000207 * cos( 1.49549061720 +    12562.56127133100*t);
+   earth_y_0+=     0.00000000212 * cos( 0.18391974040 +     7129.15868474260*t);
+   earth_y_0+=     0.00000000269 * cos( 6.23907991227 +    70755.06709173629*t);
+   earth_y_0+=     0.00000000224 * cos( 5.58680388524 +     5636.06501667660*t);
+   earth_y_0+=     0.00000000192 * cos( 0.47333599881 +    17271.88400752640*t);
+   earth_y_0+=     0.00000000241 * cos( 3.52161359834 +    15906.76412668260*t);
+   earth_y_0+=     0.00000000257 * cos( 0.87147328095 +    48739.85989708300*t);
+   earth_y_0+=     0.00000000184 * cos( 3.83487369725 +     9380.95967271720*t);
+   earth_y_0+=     0.00000000256 * cos( 0.73074058890 +      224.34479570190*t);
+   earth_y_0+=     0.00000000250 * cos( 2.67103537581 +     5934.15142954340*t);
+   earth_y_0+=     0.00000000185 * cos( 1.51386870105 +    19800.94595622480*t);
+   earth_y_0+=     0.00000000189 * cos( 2.52919302388 +    16062.18452611680*t);
+   earth_y_0+=     0.00000000175 * cos( 3.84093628240 +    11614.43329373220*t);
+   earth_y_0+=     0.00000000175 * cos( 6.15313646522 +     3744.82734573780*t);
+   earth_y_0+=     0.00000000240 * cos( 0.45238664563 +        1.48447270830*t);
+   earth_y_0+=     0.00000000241 * cos( 5.00040455119 +     8390.11035753380*t);
+   earth_y_0+=     0.00000000193 * cos( 4.71506547135 +      394.62588505920*t);
+   earth_y_0+=     0.00000000169 * cos( 0.80098424793 +      742.99006053260*t);
+   earth_y_0+=     0.00000000189 * cos( 4.03411872458 +      401.67212175720*t);
+   earth_y_0+=     0.00000000230 * cos( 0.83612842798 +     6632.00027043940*t);
+   earth_y_0+=     0.00000000227 * cos( 4.16811397278 +    34596.36465465240*t);
+   earth_y_0+=     0.00000000165 * cos( 0.39341190294 +    10454.50138660520*t);
+   earth_y_0+=     0.00000000221 * cos( 2.07224292531 +     6206.80977871580*t);
+   earth_y_0+=     0.00000000201 * cos( 6.23059857608 +     3097.88382272579*t);
+   earth_y_0+=     0.00000000205 * cos( 2.56533473350 +    16627.37091537720*t);
+   earth_y_0+=     0.00000000178 * cos( 0.01807925983 +     4797.09572892620*t);
+   earth_y_0+=     0.00000000225 * cos( 3.43883654828 +      266.60704172180*t);
+   earth_y_0+=     0.00000000167 * cos( 2.40308228314 +     5714.25397596400*t);
+   earth_y_0+=     0.00000000167 * cos( 1.16660958168 +    10344.29506538580*t);
+   earth_y_0+=     0.00000000172 * cos( 3.81893679688 +     6332.83287546320*t);
+   earth_y_0+=     0.00000000172 * cos( 5.97118942936 +     6233.31882451960*t);
+   earth_y_0+=     0.00000000214 * cos( 2.77237801335 +    24492.40611365159*t);
+   earth_y_0+=     0.00000000159 * cos( 2.33166453392 +    13198.93543929600*t);
+   earth_y_0+=     0.00000000207 * cos( 1.92518585993 +    12569.74212863460*t);
+   earth_y_0+=     0.00000000176 * cos( 2.48119329994 +    12012.58229714040*t);
+   earth_y_0+=     0.00000000196 * cos( 5.41851032427 +    11087.28512591840*t);
+   earth_y_0+=     0.00000000161 * cos( 2.18455666842 +    12132.43996210600*t);
+   earth_y_0+=     0.00000000170 * cos( 5.35052715864 +     6272.03014972750*t);
+   earth_y_0+=     0.00000000170 * cos( 4.43959906760 +     6294.12155025530*t);
+   earth_y_0+=     0.00000000168 * cos( 3.97617357428 +    17686.75284556640*t);
+   earth_y_0+=     0.00000000199 * cos( 5.15620822027 +    33326.57873317420*t);
+   earth_y_0+=     0.00000000194 * cos( 2.13712729333 +       28.44918746780*t);
+   earth_y_0+=     0.00000000144 * cos( 2.03929994576 +    16737.57723659660*t);
+   earth_y_0+=     0.00000000148 * cos( 1.10385881777 +     6851.89772401880*t);
+   earth_y_0+=     0.00000000146 * cos( 4.95201849548 +    17370.36097590980*t);
+   earth_y_0+=     0.00000000143 * cos( 5.56441147227 +     5618.31980486140*t);
+   earth_y_0+=     0.00000000140 * cos( 5.93460328071 +     2008.55753915900*t);
+   earth_y_0+=     0.00000000162 * cos( 3.21674523006 +    12551.92460598120*t);
+   earth_y_0+=     0.00000000143 * cos( 3.89365374642 +     5010.39482436420*t);
+   earth_y_0+=     0.00000000161 * cos( 1.75025798421 +    31441.67756975680*t);
+   earth_y_0+=     0.00000000151 * cos( 4.72729575304 +    16723.35014259500*t);
+   earth_y_0+=     0.00000000135 * cos( 4.79054860554 +    12431.08661994740*t);
+   earth_y_0+=     0.00000000148 * cos( 3.93657905776 +    28766.92442448400*t);
+   earth_y_0+=     0.00000000146 * cos( 6.27087503385 +    12701.21678001820*t);
+   earth_y_0+=     0.00000000137 * cos( 1.02989768135 +       11.04570026390*t);
+   earth_y_0+=     0.00000000167 * cos( 5.98371371124 +    17473.45375012840*t);
+   earth_y_0+=     0.00000000161 * cos( 1.02898808001 +    20452.86941222180*t);
+   earth_y_0+=     0.00000000124 * cos( 4.22571513469 +     6947.83189512140*t);
+   earth_y_0+=     0.00000000135 * cos( 4.95838914470 +     3531.52825029980*t);
+   earth_y_0+=     0.00000000172 * cos( 0.23035302794 +   167959.39011405829*t);
+   earth_y_0+=     0.00000000122 * cos( 4.12184077581 +     3583.34103067380*t);
+   earth_y_0+=     0.00000000170 * cos( 6.16462338453 +   333857.59012677266*t);
+   earth_y_0+=     0.00000000155 * cos( 5.62826239458 +     6058.73105428950*t);
+   earth_y_0+=     0.00000000120 * cos( 5.90379805618 +    12808.88030395680*t);
+   earth_y_0+=     0.00000000141 * cos( 3.95455987291 +   162420.55183479069*t);
+   earth_y_0+=     0.00000000145 * cos( 3.70967890118 +       95.97922721780*t);
+   earth_y_0+=     0.00000000120 * cos( 4.26997270152 +    12528.12402734700*t);
+   earth_y_0+=     0.00000000119 * cos( 5.90934719297 +    19247.37655338240*t);
+   earth_y_0+=     0.00000000117 * cos( 5.45034772919 +    11610.55195837420*t);
+   earth_y_0+=     0.00000000163 * cos( 1.54432161045 +      661.23292678100*t);
+   earth_y_0+=     0.00000000161 * cos( 2.28994816795 +    52670.06959330260*t);
+   earth_y_0+=     0.00000000163 * cos( 2.78769644779 +     9929.42622734580*t);
+   earth_y_0+=     0.00000000119 * cos( 2.71628583938 +    12249.75983032620*t);
+   earth_y_0+=     0.00000000145 * cos( 1.17170397718 +     6360.74962041940*t);
+   earth_y_0+=     0.00000000145 * cos( 2.33523694188 +     6205.40207956340*t);
+   earth_y_0+=     0.00000000159 * cos( 2.61434351364 +   228278.10465148659*t);
+   earth_y_0+=     0.00000000156 * cos( 1.00708912320 +     4307.58330413540*t);
+   earth_y_0+=     0.00000000118 * cos( 5.95164477875 +     6302.74561089119*t);
+   earth_y_0+=     0.00000000118 * cos( 3.83848144749 +     6263.40608909160*t);
+   earth_y_0+=     0.00000000135 * cos( 2.80387650268 +    19804.82729158280*t);
+   earth_y_0+=     0.00000000112 * cos( 3.88136222265 +    16840.67001081519*t);
+   earth_y_0+=     0.00000000148 * cos( 3.55909740786 +    28237.23345938940*t);
+   earth_y_0+=     0.00000000124 * cos( 5.50657775921 +    10239.58386601080*t);
+   earth_y_0+=     0.00000000137 * cos( 0.97335445483 +     6315.60840078280*t);
+   earth_y_0+=     0.00000000137 * cos( 2.53358646423 +     6250.54329920000*t);
+   earth_y_0+=     0.00000000111 * cos( 0.70362725529 +     6339.97422492700*t);
+   earth_y_0+=     0.00000000111 * cos( 2.80331366377 +     6226.17747505580*t);
+   earth_y_0+=     0.00000000139 * cos( 4.20829691432 +    11609.86254401220*t);
+   earth_y_0+=     0.00000000124 * cos( 3.12647421017 +    19402.79695281660*t);
+   earth_y_0+=     0.00000000123 * cos( 0.80538262187 +    24734.15439655740*t);
+   earth_y_0+=     0.00000000107 * cos( 0.59542543867 +    12463.05892576420*t);
+   earth_y_0+=     0.00000000140 * cos( 5.62756230388 +    38526.57435087200*t);
+   earth_y_0+=     0.00000000123 * cos( 5.71833316259 +     6305.85105144220*t);
+   earth_y_0+=     0.00000000123 * cos( 4.07179306365 +     6260.30064854060*t);
+   earth_y_0+=     0.00000000113 * cos( 0.38311301542 +     2111.65031337760*t);
+   earth_y_0+=     0.00000000117 * cos( 1.11114131260 +    24336.00539314920*t);
+   earth_y_0+=     0.00000000131 * cos( 2.04122395494 +    18415.51581209739*t);
+   earth_y_0+=     0.00000000106 * cos( 2.18835416515 +    12341.80690428090*t);
+   earth_y_0+=     0.00000000110 * cos( 3.31587494701 +     6276.44729109139*t);
+   earth_y_0+=     0.00000000110 * cos( 0.19106597205 +     6289.70440889141*t);
+   earth_y_0+=     0.00000000103 * cos( 1.18729919186 +     6324.62664097620*t);
+   earth_y_0+=     0.00000000103 * cos( 2.31964172720 +     6241.52505900660*t);
+   earth_y_0+=     0.00000000131 * cos( 5.93782050272 +     5791.41255753260*t);
+   earth_y_0+=     0.00000000107 * cos( 0.14948042255 +    83973.83535572988*t);
+   earth_y_0+=     0.00000000107 * cos( 1.18803015964 +    84019.85928049390*t);
+   earth_y_0+=     0.00000000106 * cos( 1.70601019340 +    12772.33724842000*t);
+   earth_y_0+=     0.00000000143 * cos( 6.22488817389 +     2069.00683490660*t);
+   earth_y_0+=     0.00000000143 * cos( 5.63221604750 +    18772.96147869860*t);
+   earth_y_0+=     0.00000000115 * cos( 5.35511207484 +     3641.73457151920*t);
+   earth_y_0+=     0.00000000121 * cos( 3.53431033420 +    17996.03116822220*t);
+   earth_y_0+=     0.00000000126 * cos( 5.90533164276 +    18822.92923017440*t);
+   earth_y_0+=     0.00000000115 * cos( 3.09464836385 +    12570.08385324590*t);
+   earth_y_0+=     0.00000000119 * cos( 0.26217017763 +     9814.60410029120*t);
+   earth_y_0+=     0.00000000138 * cos( 2.98838673165 +     6209.77872413240*t);
+   earth_y_0+=     0.00000000138 * cos( 0.51855418741 +     6356.37297585040*t);
+   earth_y_0+=     0.00000000100 * cos( 3.07908951574 +        5.52292430740*t);
+   earth_y_0+=     0.00000000135 * cos( 5.28292512146 +    18699.66435283960*t);
+   earth_y_0+=     0.00000000104 * cos( 0.57981411648 +    12560.62877567540*t);
+   earth_y_0+=     0.00000000098 * cos( 4.00434933508 +     5815.11085963699*t);
+   earth_y_0+=     0.00000000096 * cos( 3.47793123567 +    22743.40937951640*t);
+   earth_y_0+=     0.00000000099 * cos( 2.68455377979 +    78263.70942472259*t);
+   earth_y_0+=     0.00000000114 * cos( 2.18292460802 +     6155.05700665400*t);
+   earth_y_0+=     0.00000000114 * cos( 1.32401631104 +     6411.09469332880*t);
+   earth_y_0+=     0.00000000093 * cos( 0.34809743975 +    11190.37790013700*t);
+   earth_y_0+=     0.00000000095 * cos( 4.06228552408 +    11300.58422135640*t);
+   earth_y_0+=     0.00000000117 * cos( 2.41831042685 +    24383.07910844140*t);
+   earth_y_0+=     0.00000000110 * cos( 1.22040625314 +     6312.50535852740*t);
+   earth_y_0+=     0.00000000110 * cos( 2.28653466592 +     6253.64634145540*t);
+   earth_y_0+=     0.00000000091 * cos( 1.42377843647 +    12829.23562335560*t);
+   earth_y_0+=     0.00000000120 * cos( 4.42365311529 +    17892.93839400359*t);
+   earth_y_0+=     0.00000000090 * cos( 4.05016333707 +    15664.03552270859*t);
+   earth_y_0+=     0.00000000126 * cos( 0.70432668892 +     1293.48422358260*t);
+   earth_y_0+=     0.00000000126 * cos( 4.96989709208 +   311565.01892104017*t);
+   earth_y_0+=     0.00000000120 * cos( 0.51717377640 +    14158.74771361560*t);
+   earth_y_0+=     0.00000000110 * cos( 1.99803947181 +    16097.67995028260*t);
+   earth_y_0+=     0.00000000092 * cos( 1.75886109740 +    21424.46664430340*t);
+   earth_y_0+=     0.00000000099 * cos( 2.04878528903 +     6432.63904712600*t);
+   earth_y_0+=     0.00000000124 * cos( 5.19244899173 +   178430.04719053170*t);
+   earth_y_0+=     0.00000000088 * cos( 0.12768507477 +    15141.39079431200*t);
+   earth_y_0+=     0.00000000089 * cos( 5.78577662098 +     6751.04084034580*t);
+   earth_y_0+=     0.00000000096 * cos( 1.07414401415 +     7349.57132718140*t);
+   earth_y_0+=     0.00000000094 * cos( 1.30652024350 +     2699.73481931760*t);
+   earth_y_0+=     0.00000000087 * cos( 1.63535344372 +    23141.55838292460*t);
+   earth_y_0+=     0.00000000103 * cos( 2.80666225492 +    18202.21671665939*t);
+   earth_y_0+=     0.00000000099 * cos( 5.29408634455 +     6313.98997562640*t);
+   earth_y_0+=     0.00000000099 * cos( 4.49603988169 +     6252.16172435640*t);
+   earth_y_0+=     0.00000000090 * cos( 1.96045243600 +    58864.54391814630*t);
+   earth_y_0+=     0.00000000113 * cos( 3.06789441452 +    10557.59416082380*t);
+   earth_y_0+=     0.00000000098 * cos( 1.66301652528 +     6126.67512948900*t);
+   earth_y_0+=     0.00000000098 * cos( 1.84392439378 +     6439.47657049380*t);
+   earth_y_0+=     0.00000000086 * cos( 6.05385675766 +    23123.74586080659*t);
+   earth_y_0+=     0.00000000094 * cos( 5.54106526599 +    23539.70738633280*t);
+   earth_y_0+=     0.00000000118 * cos( 2.16762826283 +    13951.71327493390*t);
+   earth_y_0+=     0.00000000117 * cos( 2.53845687210 +    89569.99011954499*t);
+   earth_y_0+=     0.00000000083 * cos( 4.71523084157 +     8858.31494432060*t);
+   earth_y_0+=     0.00000000101 * cos( 4.94270501848 +    18429.74290609900*t);
+   earth_y_0+=     0.00000000088 * cos( 2.19663281211 +    17053.96910625319*t);
+   earth_y_0+=     0.00000000113 * cos( 0.47970289838 +    12985.63634385800*t);
+   earth_y_0+=     0.00000000082 * cos( 4.90614256431 +     5483.25472482600*t);
+   earth_y_0+=     0.00000000095 * cos( 0.81797133475 +    11925.27409260060*t);
+   earth_y_0+=     0.00000000084 * cos( 2.66665772953 +     6379.05507720920*t);
+   earth_y_0+=     0.00000000084 * cos( 0.84028318953 +     6187.09662277360*t);
+   earth_y_0+=     0.00000000093 * cos( 3.10680657507 +    35371.88726597640*t);
+   earth_y_0+=     0.00000000105 * cos( 1.20399152892 +     6434.97313107660*t);
+   earth_y_0+=     0.00000000105 * cos( 2.30294939014 +     6131.17856890620*t);
+   earth_y_0+=     0.00000000114 * cos( 2.50001433093 +     8258.56839584740*t);
+   earth_y_0+=     0.00000000084 * cos( 3.82196413310 +    32217.20018108080*t);
+   earth_y_0+=     0.00000000082 * cos( 5.40278723857 +     3958.12644117580*t);
+   earth_y_0+=     0.00000000102 * cos( 1.32703960779 +     5767.61197889840*t);
+   earth_y_0+=     0.00000000102 * cos( 2.17990131127 +     6798.53972108440*t);
+   earth_y_0+=     0.00000000092 * cos( 6.07396273786 +    18216.44381066100*t);
+   earth_y_0+=     0.00000000112 * cos( 3.91825505548 +    12188.77809206700*t);
+   earth_y_0+=     0.00000000086 * cos( 4.18439486840 +    16193.65917750039*t);
+   earth_y_0+=     0.00000000080 * cos( 2.01988251452 +    24080.03501677719*t);
+   earth_y_0+=     0.00000000093 * cos( 4.66163367881 +    23937.85638974100*t);
+   earth_y_0+=     0.00000000088 * cos( 4.57496913724 +     1573.82042409880*t);
+   earth_y_0+=     0.00000000092 * cos( 2.79375473262 +   101426.20877096949*t);
+   earth_y_0+=     0.00000000110 * cos( 5.08243853758 +    78423.70451667879*t);
+   earth_y_0+=     0.00000000087 * cos( 3.89286413327 +     1580.86666079680*t);
+   earth_y_0+=     0.00000000080 * cos( 1.70826898567 +     6485.32924516550*t);
+   earth_y_0+=     0.00000000080 * cos( 1.79867193339 +     6080.82245481730*t);
+   earth_y_0+=     0.00000000079 * cos( 6.16114738580 +    40796.33892267419*t);
+   earth_y_0+=     0.00000000079 * cos( 2.96735408932 +    36109.38220466459*t);
+   earth_y_0+=     0.00000000088 * cos( 5.58306387433 +    16943.76278503380*t);
+   earth_y_0+=     0.00000000093 * cos( 0.22422436515 +      283.85931886520*t);
+   earth_y_0+=     0.00000000088 * cos( 1.67821631800 +    12537.70251251500*t);
+   earth_y_0+=     0.00000000079 * cos( 4.76362160922 +    41194.48792608240*t);
+   earth_y_0+=     0.00000000081 * cos( 5.19149105668 +     3956.50801601940*t);
+   earth_y_0+=     0.00000000103 * cos( 3.71008010926 +    56600.27928952220*t);
+   earth_y_0+=     0.00000000079 * cos( 2.93697301370 +     5621.84292321040*t);
+   earth_y_0+=     0.00000000096 * cos( 0.41418225890 +     5888.44996493220*t);
+   earth_y_0+=     0.00000000096 * cos( 3.09275866016 +     6677.70173505060*t);
+   earth_y_0+=     0.00000000074 * cos( 4.41044639229 +     6304.41149045840*t);
+   earth_y_0+=     0.00000000074 * cos( 5.37967983395 +     6261.74020952440*t);
+   earth_y_0+=     0.00000000073 * cos( 1.05049706961 +     6695.44694686580*t);
+   earth_y_0+=     0.00000000073 * cos( 2.45644384945 +     5870.70475311700*t);
+   earth_y_0+=     0.00000000073 * cos( 1.61765616896 +    16207.88627150200*t);
+   earth_y_0+=     0.00000000085 * cos( 2.08161902252 +    12850.01101884800*t);
+   earth_y_0+=     0.00000000072 * cos( 5.03595593677 +     6312.89728814020*t);
+   earth_y_0+=     0.00000000072 * cos( 4.75417028947 +     6253.25441184259*t);
+   earth_y_0+=     0.00000000100 * cos( 3.80714752393 +     5316.10497255580*t);
+   earth_y_0+=     0.00000000085 * cos( 1.33904345426 +    12282.29238111760*t);
+   earth_y_0+=     0.00000000077 * cos( 5.07586430304 +    22345.26037610820*t);
+   earth_y_0+=     0.00000000075 * cos( 0.39510451579 +    24422.37035140730*t);
+   earth_y_0+=     0.00000000076 * cos( 0.41146256256 +    63658.87775083760*t);
+   earth_y_0+=     0.00000000087 * cos( 2.06336091493 +    18106.23748944160*t);
+   earth_y_0+=     0.00000000070 * cos( 0.44640955505 +     6442.01936777460*t);
+   earth_y_0+=     0.00000000070 * cos( 3.06053136401 +     6124.13233220820*t);
+   earth_y_0+=     0.00000000074 * cos( 5.61945589711 +    15265.88651930040*t);
+   earth_y_0+=     0.00000000078 * cos( 3.38283383527 +    23550.34405168259*t);
+   earth_y_0+=     0.00000000077 * cos( 5.00656029204 +     5237.92101380380*t);
+   earth_y_0+=     0.00000000075 * cos( 3.64547243832 +    27832.03821928320*t);
+   earth_y_0+=     0.00000000082 * cos( 3.91623659935 +     6774.73914245020*t);
+   earth_y_0+=     0.00000000067 * cos( 4.88284279928 +     7082.89697515680*t);
+   earth_y_0+=     0.00000000067 * cos( 3.24232940659 +     6171.88920770380*t);
+   earth_y_0+=     0.00000000067 * cos( 0.26461151247 +     6394.26249227900*t);
+   earth_y_0+=     0.00000000080 * cos( 3.79677325613 +    22380.75580027400*t);
+   earth_y_0+=     0.00000000070 * cos( 0.47599032576 +     6175.41232605280*t);
+   earth_y_0+=     0.00000000070 * cos( 3.03095059330 +     6390.73937393000*t);
+   earth_y_0+=     0.00000000077 * cos( 0.29137018417 +     9924.81042151060*t);
+   earth_y_0+=     0.00000000091 * cos( 0.90446873976 +     6652.77566593180*t);
+   earth_y_0+=     0.00000000091 * cos( 2.60247217930 +     5913.37603405100*t);
+   earth_y_0+=     0.00000000067 * cos( 4.84479908786 +     6300.88837210940*t);
+   earth_y_0+=     0.00000000067 * cos( 4.94532713838 +     6265.26332787340*t);
+   earth_y_0+=     0.00000000088 * cos( 3.59122542798 +    10660.68693504240*t);
+   earth_y_0+=     0.00000000066 * cos( 3.01430181986 +     5959.57043333400*t);
+   earth_y_0+=     0.00000000080 * cos( 5.49178527606 +    23646.32327890039*t);
+   earth_y_0+=     0.00000000068 * cos( 2.07232411643 +     4897.18057365520*t);
+   earth_y_0+=     0.00000000069 * cos( 0.56996801147 +     6944.30877677240*t);
+   earth_y_0+=     0.00000000073 * cos( 2.96334792367 +     6195.76764545159*t);
+   earth_y_0+=     0.00000000073 * cos( 0.54359299539 +     6370.38405453120*t);
+   earth_y_0+=     0.00000000071 * cos( 3.57166978930 +    17363.24742890899*t);
+   earth_y_0+=     0.00000000074 * cos( 1.56614709067 +    27043.50288318280*t);
+   earth_y_0+=     0.00000000066 * cos( 2.38194509538 +     6251.10354417460*t);
+   earth_y_0+=     0.00000000066 * cos( 1.12499582368 +     6315.04815580820*t);
+   earth_y_0+=     0.00000000064 * cos( 3.50738755186 +      323.50541665740*t);
+   earth_y_0+=     0.00000000078 * cos( 0.83740696355 +    42456.78404709160*t);
+   earth_y_0+=     0.00000000063 * cos( 2.70916268945 +    11823.16163945020*t);
+   earth_y_0+=     0.00000000064 * cos( 4.80330562067 +    28774.03797148480*t);
+   earth_y_0+=     0.00000000079 * cos( 5.66337403718 +    24065.80792277559*t);
+   earth_y_0+=     0.00000000072 * cos( 2.11096843463 +    18100.00325845000*t);
+   earth_y_0+=     0.00000000066 * cos( 1.21053452412 +    17576.54652434700*t);
+   earth_y_0+=     0.00000000072 * cos( 0.02703779211 +   245707.46610434429*t);
+   earth_y_0+=     0.00000000068 * cos( 0.52464381855 +    18003.14471522300*t);
+   earth_y_0+=     0.00000000068 * cos( 1.33733194362 +    28759.81087748319*t);
+   earth_y_0+=     0.00000000062 * cos( 2.25046717820 +     3855.03366695720*t);
+   earth_y_0+=     0.00000000060 * cos( 5.83008017266 +    18606.49894600020*t);
+   earth_y_0+=     0.00000000079 * cos( 2.49476757272 +    22476.73502749179*t);
+   earth_y_0+=     0.00000000073 * cos( 4.61251602377 +    23440.13773046320*t);
+   earth_y_0+=     0.00000000060 * cos( 1.65871121739 +    13088.72911807660*t);
+   earth_y_0+=     0.00000000078 * cos( 2.99014426818 +    12564.66722727450*t);
+   earth_y_0+=     0.00000000060 * cos( 2.90135289748 +     7548.64332861780*t);
+   earth_y_0+=     0.00000000073 * cos( 4.73963217435 +      735.87651353180*t);
+   earth_y_0+=     0.00000000059 * cos( 2.43330626408 +     5113.48759858280*t);
+   earth_y_0+=     0.00000000075 * cos( 2.85043788676 +    18624.88275427230*t);
+   earth_y_0+=     0.00000000069 * cos( 4.25329739931 +   167993.69463622378*t);
+   earth_y_0+=     0.00000000059 * cos( 5.13628041810 +    41990.78593289880*t);
+   earth_y_0+=     0.00000000081 * cos( 5.98297918907 +     7250.04672742700*t);
+   earth_y_0+=     0.00000000059 * cos( 2.39193379192 +     6289.14176562120*t);
+   earth_y_0+=     0.00000000059 * cos( 1.11500712714 +     6277.00993436160*t);
+   earth_y_0+=     0.00000000072 * cos( 4.87534601803 +     8982.81066930900*t);
+   earth_y_0+=     0.00000000057 * cos( 3.30911091342 +      792.77488846740*t);
+   earth_y_0+=     0.00000000057 * cos( 1.61290511224 +    10027.90319572920*t);
+   earth_y_0+=     0.00000000072 * cos( 5.73653948228 +     4598.02372748980*t);
+   earth_y_0+=     0.00000000060 * cos( 0.63449806943 +      647.01083331480*t);
+   earth_y_0+=     0.00000000066 * cos( 0.80721495515 +     5220.17080145320*t);
+   earth_y_0+=     0.00000000056 * cos( 5.24002840662 +    35050.00027447539*t);
+   earth_y_0+=     0.00000000077 * cos( 2.14165124499 +    46386.99374331120*t);
+   earth_y_0+=     0.00000000068 * cos( 1.14499338723 +     3166.41643773160*t);
+   earth_y_0+=     0.00000000056 * cos( 0.39917748570 +     6339.87911216120*t);
+   earth_y_0+=     0.00000000056 * cos( 3.10776343336 +     6226.27258782160*t);
+   earth_y_0+=     0.00000000056 * cos( 0.95649088929 +    12303.06777661000*t);
+   earth_y_0+=     0.00000000057 * cos( 1.98721525004 +     6702.00024889200*t);
+   earth_y_0+=     0.00000000057 * cos( 1.51972566902 +     5864.15145109080*t);
+   earth_y_0+=     0.00000000069 * cos( 0.33788710387 +    23536.11695768099*t);
+   earth_y_0+=     0.00000000059 * cos( 1.81421782103 +    67589.08744705719*t);
+   earth_y_0+=     0.00000000076 * cos( 0.77801824814 +    23226.83863502520*t);
+   earth_y_0+=     0.00000000059 * cos( 1.03968301691 +       15.25247118500*t);
+   earth_y_0+=     0.00000000073 * cos( 4.43868056110 +    30348.88377276700*t);
+   earth_y_0+=     0.00000000054 * cos( 4.96201305869 +    36147.40987730040*t);
+   earth_y_0+=     0.00000000070 * cos( 6.27818224660 +    12299.54465826100*t);
+   earth_y_0+=     0.00000000060 * cos( 5.08461088458 +    12171.52581492360*t);
+   earth_y_0+=     0.00000000074 * cos( 1.19928905734 +     3684.37804999020*t);
+   earth_y_0+=     0.00000000058 * cos( 1.58539308271 +     6716.78758786820*t);
+   earth_y_0+=     0.00000000061 * cos( 3.25163800829 +    27707.54249429480*t);
+   earth_y_0+=     0.00000000059 * cos( 5.76308524321 +    12164.47957822560*t);
+   earth_y_0+=     0.00000000070 * cos( 0.07217173698 +    16833.55646381439*t);
+   earth_y_0+=     0.00000000063 * cos( 2.16224668422 +     5227.62607306530*t);
+   earth_y_0+=     0.00000000063 * cos( 5.23287718161 +     5219.76176653910*t);
+   earth_y_0+=     0.00000000056 * cos( 2.45377587379 +     3178.14579056760*t);
+   earth_y_0+=     0.00000000072 * cos( 2.62702736253 +     4583.79663348820*t);
+   earth_y_0+=     0.00000000052 * cos( 0.96546586882 +     3627.50747751760*t);
+   earth_y_0+=     0.00000000052 * cos( 0.20415505351 +     6155.12431695680*t);
+   earth_y_0+=     0.00000000052 * cos( 3.30278586555 +     6411.02738302600*t);
+   earth_y_0+=     0.00000000060 * cos( 0.48342399631 +     5244.04923920100*t);
+   earth_y_0+=     0.00000000058 * cos( 2.72226132172 +    30775.48196364300*t);
+   earth_y_0+=     0.00000000052 * cos( 5.80291300969 +    30639.85663863300*t);
+   earth_y_0+=     0.00000000053 * cos( 5.75258856716 +    37724.75341974820*t);
+   earth_y_0+=     0.00000000062 * cos( 2.89766241568 +     6908.74604230380*t);
+   earth_y_0+=     0.00000000062 * cos( 0.60927850338 +     5657.40565767900*t);
+   earth_y_0+=     0.00000000050 * cos( 3.11011154958 +     2575.23909432920*t);
+   earth_y_0+=     0.00000000054 * cos( 1.15676928987 +    18208.34994259200*t);
+   earth_y_0+=     0.00000000050 * cos( 3.12376167242 +     5881.40372823420*t);
+   earth_y_0+=     0.00000000050 * cos( 0.38317924664 +     6684.74797174860*t);
+   earth_y_0+=     0.00000000068 * cos( 4.61950794623 +    96677.89886304220*t);
+   earth_y_0+=     0.00000000054 * cos( 0.77163765121 +    41654.96311596780*t);
+   earth_y_0+=     0.00000000066 * cos( 5.13003693613 +    60530.48898574180*t);
+   earth_y_0+=     0.00000000050 * cos( 1.64593222191 +    43738.80234596580*t);
+   earth_y_0+=     0.00000000048 * cos( 5.44674337925 +    25685.87280280800*t);
+   earth_y_0+=     0.00000000059 * cos( 2.71376008731 +     3981.49003408200*t);
+   earth_y_0+=     0.00000000058 * cos( 0.74658643142 +    14169.79356223040*t);
+   earth_y_0+=     0.00000000051 * cos( 1.43461659855 +     7668.97112632760*t);
+   earth_y_0+=     0.00000000048 * cos( 3.23576011756 +     3646.35037735440*t);
+   earth_y_0+=     0.00000000054 * cos( 4.50803792908 +    39302.09696219600*t);
+   earth_y_0+=     0.00000000057 * cos( 2.62704971421 +      799.82112516540*t);
+   earth_y_0+=     0.00000000061 * cos( 3.28745263972 +    44034.12758953940*t);
+   earth_y_0+=     0.00000000049 * cos( 2.15365169571 +    96562.99901809469*t);
+   earth_y_0+=     0.00000000050 * cos( 3.60626874009 +    17583.66007134780*t);
+   earth_y_0+=     0.00000000051 * cos( 1.16028405425 +    38650.17350619900*t);
+   earth_y_0+=     0.00000000053 * cos( 1.90122336672 +    11080.17157891760*t);
+   earth_y_0+=     0.00000000059 * cos( 0.68886289824 +    62883.35513951360*t);
+   earth_y_0+=     0.00000000046 * cos( 0.83996053249 +    25287.72379939980*t);
+   earth_y_0+=     0.00000000049 * cos( 3.59486880927 +     5490.30096152400*t);
+   earth_y_0+=     0.00000000058 * cos( 5.95314247069 +     6489.77658728800*t);
+   earth_y_0+=     0.00000000058 * cos( 3.83698375555 +     6076.37511269480*t);
+   earth_y_0+=     0.00000000047 * cos( 4.43378504018 +      170.67287061920*t);
+   earth_y_0+=     0.00000000045 * cos( 5.81363953764 +    33318.81606245820*t);
+   earth_y_0+=     0.00000000046 * cos( 3.19651245272 +     4480.70385926960*t);
+   earth_y_0+=     0.00000000060 * cos( 3.92066214096 +    28313.28880466100*t);
+   earth_y_0+=     0.00000000061 * cos( 3.53510681440 +      491.66329245880*t);
+   earth_y_0+=     0.00000000051 * cos( 6.00904914233 +     3428.43547608120*t);
+   earth_y_0+=     0.00000000060 * cos( 3.51283911810 +    50317.20343953080*t);
+   earth_y_0+=     0.00000000054 * cos( 1.22650408797 +     6006.33007812700*t);
+   earth_y_0+=     0.00000000054 * cos( 2.28043683109 +     6559.82162185580*t);
+   earth_y_0+=     0.00000000052 * cos( 3.85271761797 +   244287.60000722769*t);
+   earth_y_0+=     0.00000000046 * cos( 3.84965479184 +     4384.72463205180*t);
+   earth_y_0+=     0.00000000045 * cos( 0.29195078507 +    44136.95134937400*t);
+   earth_y_0+=     0.00000000044 * cos( 2.71274650311 +     2301.58581590939*t);
+   earth_y_0+=     0.00000000048 * cos( 4.64164514773 +    35579.69123957000*t);
+   earth_y_0+=     0.00000000051 * cos( 4.91887651070 +    18326.65013188039*t);
+   earth_y_0+=     0.00000000046 * cos( 3.15339406184 +    16703.06213349900*t);
+   earth_y_0+=     0.00000000048 * cos( 2.39477993061 +     5209.46682580060*t);
+   earth_y_0+=     0.00000000046 * cos( 2.47672465258 +    12555.10599971890*t);
+   earth_y_0+=     0.00000000052 * cos( 4.64446990453 +     5992.10298412540*t);
+   earth_y_0+=     0.00000000052 * cos( 5.14565632171 +     6574.04871585740*t);
+   earth_y_0+=     0.00000000044 * cos( 3.80836113401 +    69941.95360082899*t);
+   earth_y_0+=     0.00000000050 * cos( 3.54529295551 +    22030.21295466960*t);
+   earth_y_0+=     0.00000000043 * cos( 1.20775471242 +    20894.77567920879*t);
+   earth_y_0+=     0.00000000044 * cos( 2.61939991177 +      533.21408344360*t);
+   earth_y_0+=     0.00000000045 * cos( 4.65113002963 +    46360.69542351140*t);
+   earth_y_0+=     0.00000000047 * cos( 2.17023181784 +    45585.17281218740*t);
+   earth_y_0+=     0.00000000043 * cos( 1.54528183996 +     8401.83971036980*t);
+   earth_y_0+=     0.00000000054 * cos( 3.62371389676 +    24279.10701821359*t);
+   earth_y_0+=     0.00000000042 * cos( 5.13054174413 +    27433.88921587499*t);
+   earth_y_0+=     0.00000000042 * cos( 1.32093818083 +    21548.96236929180*t);
+   earth_y_0+=     0.00000000045 * cos( 3.22031193914 +    71519.29714327680*t);
+   earth_y_0+=     0.00000000043 * cos( 2.75467576214 +    27278.46881644080*t);
+   earth_y_0+=     0.00000000042 * cos( 2.68353478578 +    28230.18722269139*t);
+   earth_y_0+=     0.00000000051 * cos( 2.95623880108 +     7761.94242405580*t);
+   earth_y_0+=     0.00000000041 * cos( 5.14394847571 +     3735.23831175900*t);
+   earth_y_0+=     0.00000000042 * cos( 3.66392136716 +    26084.02180621620*t);
+   earth_y_0+=     0.00000000047 * cos( 6.03338630990 +     4708.23004870920*t);
+   earth_y_0+=     0.00000000056 * cos( 5.67442109710 +    88041.95685927900*t);
+   earth_y_0+=     0.00000000056 * cos( 5.36734540218 +    28663.83165026540*t);
+   earth_y_0+=     0.00000000040 * cos( 1.07842393153 +    28628.33622609960*t);
+   earth_y_0+=     0.00000000047 * cos( 5.61045374376 +   316428.22867391503*t);
+   earth_y_0+=     0.00000000056 * cos( 1.28041118541 +    83659.27687239899*t);
+   earth_y_0+=     0.00000000039 * cos( 0.71682305362 +    34570.06633485260*t);
+   earth_y_0+=     0.00000000056 * cos( 1.60418809786 +    82534.40362061159*t);
+   earth_y_0+=     0.00000000052 * cos( 1.15000260118 +    40398.18991926600*t);
+   earth_y_0+=     0.00000000048 * cos( 0.93182108526 +    12662.13092720060*t);
+   earth_y_0+=     0.00000000041 * cos( 4.95215656542 +    37057.57749256620*t);
+   earth_y_0+=     0.00000000052 * cos( 2.84085536047 +    12571.67462429020*t);
+   earth_y_0+=     0.00000000043 * cos( 3.12714109961 +     7322.10246078180*t);
+   earth_y_0+=     0.00000000041 * cos( 1.98689255259 +      664.75604513000*t);
+   earth_y_0+=     0.00000000043 * cos( 4.45441019609 +    19004.64794940840*t);
+   earth_y_0+=     0.00000000044 * cos( 3.63007864944 +     6140.89722295520*t);
+   earth_y_0+=     0.00000000044 * cos( 6.16004757680 +     6425.25447702760*t);
+   earth_y_0+=     0.00000000044 * cos( 3.57757676616 +    49515.38250840700*t);
+   earth_y_0+=     0.00000000039 * cos( 5.09443272912 +    22594.05489571199*t);
+   earth_y_0+=     0.00000000055 * cos( 5.97876206707 +    80181.53746683979*t);
+   earth_y_0+=     0.00000000054 * cos( 3.12084664124 +     6006.04085625000*t);
+   earth_y_0+=     0.00000000054 * cos( 0.38609427782 +     6560.11084373280*t);
+   earth_y_0+=     0.00000000045 * cos( 0.83262497188 +   142861.39123625819*t);
+   earth_y_0+=     0.00000000047 * cos( 5.83395711112 +    24499.51966065240*t);
+   earth_y_0+=     0.00000000043 * cos( 6.07854396898 +     7026.06591052400*t);
+   earth_y_0+=     0.00000000047 * cos( 2.57216284949 +    11919.14086666800*t);
+   earth_y_0+=     0.00000000047 * cos( 2.93748065634 +    24485.29256665079*t);
+   earth_y_0+=     0.00000000050 * cos( 3.13873979564 +    92747.68916682259*t);
+   earth_y_0+=     0.00000000043 * cos( 4.24359473122 +     6359.34192126700*t);
+   earth_y_0+=     0.00000000049 * cos( 2.11824959319 +    66813.56483573320*t);
+   earth_y_0+=     0.00000000042 * cos( 3.13003412642 +     6231.79551212899*t);
+   earth_y_0+=     0.00000000042 * cos( 0.37690679264 +     6334.35618785381*t);
+   earth_y_0+=     0.00000000052 * cos( 1.96232281613 +    11904.91877320180*t);
+   earth_y_0+=     0.00000000040 * cos( 6.22187276386 +     7075.85073845880*t);
+   earth_y_0+=     0.00000000040 * cos( 4.78395267597 +     7122.04513774180*t);
+   earth_y_0+=     0.00000000040 * cos( 5.00617355027 +     5444.10656224100*t);
+   earth_y_0+=     0.00000000038 * cos( 0.67444569484 +    44933.24935619040*t);
+   earth_y_0+=     0.00000000037 * cos( 4.80355711691 +     6850.79448772180*t);
+   earth_y_0+=     0.00000000037 * cos( 4.98656910933 +     5715.35721226100*t);
+   earth_y_0+=     0.00000000050 * cos( 0.35643037555 +    21947.11137270000*t);
+   earth_y_0+=     0.00000000039 * cos( 2.97136332005 +    21150.81336588360*t);
+   earth_y_0+=     0.00000000042 * cos( 2.21951419227 +     6328.21706962800*t);
+   earth_y_0+=     0.00000000042 * cos( 1.28742672679 +     6237.93463035480*t);
+   earth_y_0+=     0.00000000051 * cos( 4.42510682790 +    60284.41001526290*t);
+   earth_y_0+=     0.00000000040 * cos( 5.21447072458 +     3215.13638064320*t);
+   earth_y_0+=     0.00000000039 * cos( 3.85122558089 +    30774.50164257480*t);
+   earth_y_0+=     0.00000000046 * cos( 3.45581538922 +    12146.66705610760*t);
+   earth_y_0+=     0.00000000050 * cos( 0.05709243771 +    84334.41776382479*t);
+
+   double earth_y_1=0.0;
+   earth_y_1+=     0.00093046324 * cos( 0.00000000000 +        0.00000000000*t);
+   earth_y_1+=     0.00051506609 * cos( 4.43180499286 +    12566.15169998280*t);
+   earth_y_1+=     0.00001290800 * cos( 4.38860548540 +    18849.22754997420*t);
+   earth_y_1+=     0.00000464550 * cos( 5.82729912952 +     6283.07584999140*t);
+   earth_y_1+=     0.00000212689 * cos( 0.16300556918 +     6279.55273164240*t);
+   earth_y_1+=     0.00000212533 * cos( 3.34400595407 +     6286.59896834040*t);
+   earth_y_1+=     0.00000062345 * cos( 5.07377354827 +     4705.73230754360*t);
+   earth_y_1+=     0.00000059794 * cos( 2.24100907272 +     6256.77753019160*t);
+   earth_y_1+=     0.00000059441 * cos( 1.26619990626 +     6309.37416979120*t);
+   earth_y_1+=     0.00000048987 * cos( 3.64166577835 +      775.52261132400*t);
+   earth_y_1+=     0.00000042746 * cos( 5.15177425824 +     1059.38193018920*t);
+   earth_y_1+=     0.00000046642 * cos( 4.71575215800 +     7860.41939243920*t);
+   earth_y_1+=     0.00000037274 * cos( 0.66041323804 +     5753.38488489680*t);
+   earth_y_1+=     0.00000035625 * cos( 6.15462846966 +     5884.92684658320*t);
+   earth_y_1+=     0.00000035373 * cos( 2.90052012152 +     6812.76681508600*t);
+   earth_y_1+=     0.00000032157 * cos( 3.62669702015 +     6681.22485339960*t);
+   earth_y_1+=     0.00000028763 * cos( 4.34539355621 +    25132.30339996560*t);
+   earth_y_1+=     0.00000028447 * cos( 5.86215103214 +     6127.65545055720*t);
+   earth_y_1+=     0.00000027502 * cos( 3.92641632542 +     6438.49624942560*t);
+   earth_y_1+=     0.00000024830 * cos( 1.35190115953 +     5486.77784317500*t);
+   earth_y_1+=     0.00000020618 * cos( 2.14772396647 +     7079.37385680780*t);
+   earth_y_1+=     0.00000019531 * cos( 1.29958310594 +     5507.55323866740*t);
+   earth_y_1+=     0.00000018447 * cos( 6.17330532651 +    11790.62908865880*t);
+   earth_y_1+=     0.00000016520 * cos( 4.63328039489 +    11506.76976979360*t);
+   earth_y_1+=     0.00000016754 * cos( 2.24862220282 +     7058.59846131540*t);
+   earth_y_1+=     0.00000014560 * cos( 4.40425761361 +     6290.18939699220*t);
+   earth_y_1+=     0.00000014426 * cos( 5.39226330857 +     6275.96230299060*t);
+   earth_y_1+=     0.00000014186 * cos( 6.18547219419 +      796.29800681640*t);
+   earth_y_1+=     0.00000012666 * cos( 5.73231393752 +     4694.00295470760*t);
+   earth_y_1+=     0.00000012966 * cos( 5.10338935824 +        7.11354700080*t);
+   earth_y_1+=     0.00000010100 * cos( 1.72464513593 +     3738.76143010800*t);
+   earth_y_1+=     0.00000009427 * cos( 3.08115743044 +     6069.77675455340*t);
+   earth_y_1+=     0.00000010425 * cos( 2.69531177468 +     6282.09552892320*t);
+   earth_y_1+=     0.00000010425 * cos( 0.81162914438 +     6284.05617105960*t);
+   earth_y_1+=     0.00000009576 * cos( 0.32459396888 +     4136.91043351620*t);
+   earth_y_1+=     0.00000009043 * cos( 0.43888880403 +     6496.37494542940*t);
+   earth_y_1+=     0.00000008545 * cos( 4.77137497641 +     1194.44701022460*t);
+   earth_y_1+=     0.00000008199 * cos( 5.96838316358 +     1589.07289528380*t);
+   earth_y_1+=     0.00000006291 * cos( 3.18312019173 +     8429.24126646660*t);
+   earth_y_1+=     0.00000006133 * cos( 1.78033290355 +     8827.39026987480*t);
+   earth_y_1+=     0.00000005524 * cos( 3.79842135217 +     4933.20844033260*t);
+   earth_y_1+=     0.00000007316 * cos( 5.20404355405 +     4535.05943692440*t);
+   earth_y_1+=     0.00000005146 * cos( 3.57146870511 +    11769.85369316640*t);
+   earth_y_1+=     0.00000006461 * cos( 2.67840586109 +     5088.62883976680*t);
+   earth_y_1+=     0.00000005270 * cos( 0.94095218349 +     3154.68708489560*t);
+   earth_y_1+=     0.00000004829 * cos( 5.32876538500 +     6040.34724601740*t);
+   earth_y_1+=     0.00000004985 * cos( 3.05697984231 +    12569.67481833180*t);
+   earth_y_1+=     0.00000004709 * cos( 2.39292099451 +     5331.35744374080*t);
+   earth_y_1+=     0.00000004601 * cos( 4.45296532710 +     6525.80445396540*t);
+   earth_y_1+=     0.00000005607 * cos( 0.98720797401 +     5729.50644714900*t);
+   earth_y_1+=     0.00000004081 * cos( 5.94191743597 +     7632.94325965020*t);
+   earth_y_1+=     0.00000004394 * cos( 4.75948184840 +    12559.03815298200*t);
+   earth_y_1+=     0.00000005053 * cos( 4.55162942081 +     8031.09226305840*t);
+   earth_y_1+=     0.00000005077 * cos( 0.73549620962 +     7477.52286021600*t);
+   earth_y_1+=     0.00000003779 * cos( 1.07385171253 +     7234.79425624200*t);
+   earth_y_1+=     0.00000004974 * cos( 2.50029147401 +     6836.64525283380*t);
+   earth_y_1+=     0.00000003614 * cos( 2.52455299148 +     3930.20969621960*t);
+   earth_y_1+=     0.00000003482 * cos( 3.82652198801 +    10977.07880469900*t);
+   earth_y_1+=     0.00000003576 * cos( 5.30219799481 +     4164.31198961300*t);
+   earth_y_1+=     0.00000003376 * cos( 5.00213145251 +    11371.70468975820*t);
+   earth_y_1+=     0.00000004290 * cos( 3.36736739298 +     1592.59601363280*t);
+   earth_y_1+=     0.00000003406 * cos( 0.83525752781 +     3128.38876509580*t);
+   earth_y_1+=     0.00000003124 * cos( 3.17143982984 +     7084.89678111520*t);
+   earth_y_1+=     0.00000002752 * cos( 3.67389621489 +     7342.45778018060*t);
+   earth_y_1+=     0.00000002778 * cos( 0.26817898615 +     1748.01641306700*t);
+   earth_y_1+=     0.00000002694 * cos( 4.70989138148 +      529.69096509460*t);
+   earth_y_1+=     0.00000002820 * cos( 0.23146914320 +     5223.69391980220*t);
+   earth_y_1+=     0.00000003334 * cos( 3.48823993118 +      801.82093112380*t);
+   earth_y_1+=     0.00000003328 * cos( 5.83812544995 +      155.42039943420*t);
+   earth_y_1+=     0.00000003008 * cos( 0.67203387304 +     8635.94200376320*t);
+   earth_y_1+=     0.00000002348 * cos( 5.19665664165 +     2146.16541647520*t);
+   earth_y_1+=     0.00000002307 * cos( 2.98930142346 +    17260.15465469040*t);
+   earth_y_1+=     0.00000002778 * cos( 6.17392405934 +     5481.25491886760*t);
+   earth_y_1+=     0.00000002046 * cos( 1.96605971472 +     1990.74501704100*t);
+   earth_y_1+=     0.00000002670 * cos( 1.00151747879 +       26.29831979980*t);
+   earth_y_1+=     0.00000002000 * cos( 3.03124266219 +      951.71840625060*t);
+   earth_y_1+=     0.00000002162 * cos( 0.33268106563 +      536.80451209540*t);
+   earth_y_1+=     0.00000001830 * cos( 4.43806423544 +      553.56940284240*t);
+   earth_y_1+=     0.00000001762 * cos( 3.84931185617 +     4690.47983635860*t);
+   earth_y_1+=     0.00000002033 * cos( 0.80949219978 +      522.57741809380*t);
+   earth_y_1+=     0.00000002430 * cos( 1.64090263419 +     1349.86740965880*t);
+   earth_y_1+=     0.00000001902 * cos( 2.66567222573 +     9437.76293488700*t);
+   earth_y_1+=     0.00000001940 * cos( 3.91085273629 +     5216.58037280140*t);
+   earth_y_1+=     0.00000001708 * cos( 5.00417878382 +    10988.80815753500*t);
+   earth_y_1+=     0.00000001762 * cos( 2.39542835845 +      398.14900340820*t);
+   earth_y_1+=     0.00000001790 * cos( 3.48655562486 +     5230.80746680300*t);
+   earth_y_1+=     0.00000001524 * cos( 5.81370227284 +      426.59819087600*t);
+   earth_y_1+=     0.00000001555 * cos( 2.09924417685 +    13095.84266507740*t);
+   earth_y_1+=     0.00000001779 * cos( 0.23737949225 +    12562.62858163380*t);
+   earth_y_1+=     0.00000001450 * cos( 5.71201368260 +     6262.72053059260*t);
+   earth_y_1+=     0.00000001450 * cos( 4.07811254364 +     6303.43116939020*t);
+   earth_y_1+=     0.00000001768 * cos( 0.08826500577 +    10973.55568635000*t);
+   earth_y_1+=     0.00000001373 * cos( 5.60523185373 +     7875.67186362420*t);
+   earth_y_1+=     0.00000001428 * cos( 3.29442777483 +    12721.57209941700*t);
+   earth_y_1+=     0.00000001247 * cos( 0.29212365805 +     2118.76386037840*t);
+   earth_y_1+=     0.00000001230 * cos( 2.53131850803 +    18319.53658487960*t);
+   earth_y_1+=     0.00000001141 * cos( 1.92147452671 +     1066.49547719000*t);
+   earth_y_1+=     0.00000001140 * cos( 1.21410806683 +     2942.46342329160*t);
+   earth_y_1+=     0.00000001045 * cos( 2.17535721280 +      242.72860397400*t);
+   earth_y_1+=     0.00000000973 * cos( 4.95437649984 +      639.89728631400*t);
+   earth_y_1+=     0.00000000969 * cos( 1.29405045734 +     5643.17856367740*t);
+   earth_y_1+=     0.00000001064 * cos( 4.40566643488 +    15720.83878487840*t);
+   earth_y_1+=     0.00000000962 * cos( 0.55758587779 +     2388.89402044920*t);
+   earth_y_1+=     0.00000000943 * cos( 4.22062437022 +    20426.57109242200*t);
+   earth_y_1+=     0.00000000869 * cos( 5.80860615243 +    16496.36139620240*t);
+   earth_y_1+=     0.00000000874 * cos( 1.43036958888 +    10575.40668294180*t);
+   earth_y_1+=     0.00000001088 * cos( 5.77656856686 +     5760.49843189760*t);
+   earth_y_1+=     0.00000000889 * cos( 4.77820137734 +     6288.59877429880*t);
+   earth_y_1+=     0.00000000889 * cos( 5.01192484889 +     6277.55292568400*t);
+   earth_y_1+=     0.00000001018 * cos( 2.66299766067 +    12539.85338018300*t);
+   earth_y_1+=     0.00000001017 * cos( 4.02093017183 +     6805.65326808520*t);
+   earth_y_1+=     0.00000000794 * cos( 3.79329079940 +    14314.16811304980*t);
+   earth_y_1+=     0.00000000795 * cos( 3.39951828424 +     6279.78949257360*t);
+   earth_y_1+=     0.00000000795 * cos( 0.10742263482 +     6286.36220740920*t);
+   earth_y_1+=     0.00000000759 * cos( 4.98301749094 +     1577.34354244780*t);
+   earth_y_1+=     0.00000000757 * cos( 3.81754073272 +     5856.47765911540*t);
+   earth_y_1+=     0.00000000800 * cos( 3.18366218591 +    12029.34718788740*t);
+   earth_y_1+=     0.00000000716 * cos( 2.28687269516 +     9623.68827669120*t);
+   earth_y_1+=     0.00000000698 * cos( 4.94363229969 +     6148.01076995600*t);
+   earth_y_1+=     0.00000000907 * cos( 0.70119368826 +    13367.97263110660*t);
+   earth_y_1+=     0.00000000672 * cos( 3.99153941403 +    12964.30070339100*t);
+   earth_y_1+=     0.00000000678 * cos( 4.84647216304 +     6418.14093002680*t);
+   earth_y_1+=     0.00000000671 * cos( 6.00179402654 +     6709.67404086740*t);
+   earth_y_1+=     0.00000000657 * cos( 1.80163596347 +    13119.72110282519*t);
+   earth_y_1+=     0.00000000662 * cos( 4.04031103640 +    11015.10647733480*t);
+   earth_y_1+=     0.00000000744 * cos( 2.68741160780 +    12043.57428188900*t);
+   earth_y_1+=     0.00000000673 * cos( 2.45322928062 +    14712.31711645800*t);
+   earth_y_1+=     0.00000000696 * cos( 2.53877129238 +    16730.46368959580*t);
+   earth_y_1+=     0.00000000671 * cos( 0.40543909773 +    13517.87010623340*t);
+   earth_y_1+=     0.00000000604 * cos( 5.39462125376 +    14143.49524243060*t);
+   earth_y_1+=     0.00000000754 * cos( 5.24655445271 +    13916.01910964160*t);
+   earth_y_1+=     0.00000000629 * cos( 3.42493725594 +     4686.88940770680*t);
+   earth_y_1+=     0.00000000610 * cos( 5.19955844918 +     6819.88036208680*t);
+   earth_y_1+=     0.00000000667 * cos( 4.87151372910 +     3634.62102451840*t);
+   earth_y_1+=     0.00000000609 * cos( 4.57544043999 +     5746.27133789600*t);
+   earth_y_1+=     0.00000000620 * cos( 4.30310535016 +    31415.37924995700*t);
+   earth_y_1+=     0.00000000553 * cos( 0.78994516254 +     6306.95428773920*t);
+   earth_y_1+=     0.00000000553 * cos( 2.71699575652 +     6259.19741224360*t);
+   earth_y_1+=     0.00000000524 * cos( 3.96543945899 +    12139.55350910680*t);
+   earth_y_1+=     0.00000000514 * cos( 3.87793148429 +     4701.11650170840*t);
+   earth_y_1+=     0.00000000553 * cos( 4.81568395217 +    12036.46073488820*t);
+   earth_y_1+=     0.00000000495 * cos( 2.00987395637 +    13362.44970679920*t);
+   earth_y_1+=     0.00000000562 * cos( 4.76145666018 +     4292.33083295040*t);
+   earth_y_1+=     0.00000000481 * cos( 2.57118063835 +     6286.95718534940*t);
+   earth_y_1+=     0.00000000481 * cos( 0.93576028071 +     6279.19451463340*t);
+   earth_y_1+=     0.00000000635 * cos( 3.19879108676 +    84334.66158130829*t);
+   earth_y_1+=     0.00000000445 * cos( 1.68959289911 +    11499.65622279280*t);
+   earth_y_1+=     0.00000000596 * cos( 1.08038043695 +    10447.38783960440*t);
+   earth_y_1+=     0.00000000452 * cos( 3.98330238308 +    12573.26524698360*t);
+   earth_y_1+=     0.00000000499 * cos( 5.79985861523 +    12410.73130054860*t);
+   earth_y_1+=     0.00000000440 * cos( 0.58121786395 +    24072.92146977640*t);
+   earth_y_1+=     0.00000000575 * cos( 1.39410252291 +    12592.45001978260*t);
+   earth_y_1+=     0.00000000404 * cos( 0.56451183179 +     5863.59120611620*t);
+   earth_y_1+=     0.00000000398 * cos( 4.20035189119 +     8273.82086703240*t);
+   earth_y_1+=     0.00000000418 * cos( 2.76069368004 +    10177.25767953360*t);
+   earth_y_1+=     0.00000000379 * cos( 3.05384870302 +     6386.16862421000*t);
+   earth_y_1+=     0.00000000379 * cos( 0.45309221604 +     6179.98307577280*t);
+   earth_y_1+=     0.00000000377 * cos( 4.33269290931 +     7872.14874527520*t);
+   earth_y_1+=     0.00000000419 * cos( 4.05140116867 +     5327.47610838280*t);
+   earth_y_1+=     0.00000000383 * cos( 5.11659434377 +     6319.10371666880*t);
+   earth_y_1+=     0.00000000383 * cos( 4.67353188247 +     6247.04798331399*t);
+   earth_y_1+=     0.00000000457 * cos( 0.88573283839 +     6489.26139842860*t);
+   earth_y_1+=     0.00000000457 * cos( 2.62120808067 +     6076.89030155420*t);
+   earth_y_1+=     0.00000000378 * cos( 5.80861101902 +     7238.67559160000*t);
+   earth_y_1+=     0.00000000344 * cos( 4.99907349425 +    18451.07854656599*t);
+   earth_y_1+=     0.00000000334 * cos( 0.56215917990 +    18073.70493865020*t);
+   earth_y_1+=     0.00000000453 * cos( 5.43513758472 +     2787.04302385740*t);
+   earth_y_1+=     0.00000000364 * cos( 0.67991079022 +     6321.10352262720*t);
+   earth_y_1+=     0.00000000364 * cos( 2.82703012884 +     6245.04817735560*t);
+   earth_y_1+=     0.00000000342 * cos( 4.52940007807 +     6303.85124548380*t);
+   earth_y_1+=     0.00000000342 * cos( 5.26072614817 +     6262.30045449900*t);
+   earth_y_1+=     0.00000000319 * cos( 4.83173153249 +    11926.25441366880*t);
+   earth_y_1+=     0.00000000316 * cos( 1.17362065540 +    77713.77146812050*t);
+   earth_y_1+=     0.00000000318 * cos( 6.03589239585 +     6311.52503745920*t);
+   earth_y_1+=     0.00000000318 * cos( 3.75423383039 +     6254.62666252360*t);
+   earth_y_1+=     0.00000000315 * cos( 1.09950918957 +    12323.42309600880*t);
+   earth_y_1+=     0.00000000310 * cos( 4.02477182017 +    10213.28554621100*t);
+   earth_y_1+=     0.00000000311 * cos( 5.94074472453 +     2352.86615377180*t);
+   earth_y_1+=     0.00000000304 * cos( 0.43890915612 +    12779.45079542080*t);
+   earth_y_1+=     0.00000000343 * cos( 1.36470897271 +     6267.82337880639*t);
+   earth_y_1+=     0.00000000343 * cos( 2.14223194635 +     6298.32832117640*t);
+   earth_y_1+=     0.00000000398 * cos( 4.83558508405 +     1551.04522264800*t);
+   earth_y_1+=     0.00000000358 * cos( 2.52713972155 +     1052.26838318840*t);
+   earth_y_1+=     0.00000000300 * cos( 0.20527323254 +    90279.92316810328*t);
+   earth_y_1+=     0.00000000297 * cos( 2.07974536467 +     5429.87946823940*t);
+   earth_y_1+=     0.00000000332 * cos( 2.13791551771 +      213.29909543800*t);
+   earth_y_1+=     0.00000000362 * cos( 1.16895887442 +    10021.83728009940*t);
+   earth_y_1+=     0.00000000334 * cos( 1.16416729842 +    17782.73207278420*t);
+   earth_y_1+=     0.00000000295 * cos( 2.48770580169 +     6702.56049386660*t);
+   earth_y_1+=     0.00000000321 * cos( 0.66776723307 +    17796.95916678580*t);
+   earth_y_1+=     0.00000000369 * cos( 4.78231597870 +     4732.03062734340*t);
+   earth_y_1+=     0.00000000288 * cos( 3.73055532740 +      220.41264243880*t);
+   earth_y_1+=     0.00000000263 * cos( 2.60909847321 +      955.59974160860*t);
+   earth_y_1+=     0.00000000305 * cos( 0.57265325631 +     5642.19824260920*t);
+   earth_y_1+=     0.00000000258 * cos( 4.38103579313 +      206.18554843720*t);
+   earth_y_1+=     0.00000000287 * cos( 0.31572653335 +     9917.69687450980*t);
+   earth_y_1+=     0.00000000261 * cos( 1.93782216260 +    16200.77272450120*t);
+   earth_y_1+=     0.00000000234 * cos( 2.50307529331 +     1596.18644228460*t);
+   earth_y_1+=     0.00000000282 * cos( 2.93181184874 +    13341.67431130680*t);
+   earth_y_1+=     0.00000000262 * cos( 5.15409122759 +    14945.31617355440*t);
+   earth_y_1+=     0.00000000276 * cos( 0.59107321564 +    12168.00269657460*t);
+   earth_y_1+=     0.00000000217 * cos( 2.74597813026 +     2544.31441988340*t);
+   earth_y_1+=     0.00000000259 * cos( 5.53112324599 +     3894.18182954220*t);
+   earth_y_1+=     0.00000000216 * cos( 4.02991604351 +     3185.19202726560*t);
+   earth_y_1+=     0.00000000250 * cos( 3.04516511604 +     6923.95345737360*t);
+   earth_y_1+=     0.00000000193 * cos( 0.78779752820 +      640.87760738220*t);
+   earth_y_1+=     0.00000000203 * cos( 2.82417195790 +     6062.66320755260*t);
+   earth_y_1+=     0.00000000203 * cos( 0.68276896116 +     6503.48849243020*t);
+   earth_y_1+=     0.00000000188 * cos( 4.83298725434 +    18209.33026366019*t);
+   earth_y_1+=     0.00000000184 * cos( 2.63875516800 +    12566.21901028560*t);
+   earth_y_1+=     0.00000000213 * cos( 4.09959554852 +     9779.10867612540*t);
+   earth_y_1+=     0.00000000180 * cos( 4.19738483213 +      853.19638175200*t);
+   earth_y_1+=     0.00000000248 * cos( 5.62050900898 +    17789.84561978500*t);
+   earth_y_1+=     0.00000000187 * cos( 3.02088056063 +     6276.02961329339*t);
+   earth_y_1+=     0.00000000187 * cos( 0.48606035843 +     6290.12208668940*t);
+   earth_y_1+=     0.00000000168 * cos( 6.02158735674 +      110.20632121940*t);
+   earth_y_1+=     0.00000000184 * cos( 3.08285813784 +    12566.08438968000*t);
+   earth_y_1+=     0.00000000230 * cos( 0.66694562430 +    12567.13202105100*t);
+   earth_y_1+=     0.00000000166 * cos( 5.90725172950 +      433.71173787680*t);
+   earth_y_1+=     0.00000000184 * cos( 1.80970202711 +     6037.24420376200*t);
+   earth_y_1+=     0.00000000159 * cos( 0.31775301914 +     3496.03282613400*t);
+   earth_y_1+=     0.00000000179 * cos( 1.27622455391 +        3.52311834900*t);
+   earth_y_1+=     0.00000000156 * cos( 0.77329171048 +     2648.45482547300*t);
+   earth_y_1+=     0.00000000182 * cos( 0.19927362554 +     5849.36411211460*t);
+   earth_y_1+=     0.00000000175 * cos( 1.69691956469 +     6528.90749622080*t);
+   earth_y_1+=     0.00000000154 * cos( 5.68366565747 +     6112.40297937220*t);
+   earth_y_1+=     0.00000000154 * cos( 4.10646056877 +     6453.74872061060*t);
+   earth_y_1+=     0.00000000162 * cos( 0.40509792796 +    18875.52586977400*t);
+   earth_y_1+=     0.00000000163 * cos( 2.96167567065 +    12352.85260454480*t);
+   earth_y_1+=     0.00000000146 * cos( 2.48714617179 +      149.56319713460*t);
+   earth_y_1+=     0.00000000143 * cos( 0.06272029982 +    19651.04848109800*t);
+   earth_y_1+=     0.00000000144 * cos( 2.19098385924 +    18635.92845453620*t);
+   earth_y_1+=     0.00000000141 * cos( 5.63574749070 +    10969.96525769820*t);
+   earth_y_1+=     0.00000000146 * cos( 0.42084370397 +    13760.59871020740*t);
+   earth_y_1+=     0.00000000161 * cos( 6.06631099179 +     8662.24032356300*t);
+   earth_y_1+=     0.00000000142 * cos( 2.91677439721 +     4157.19844261220*t);
+   earth_y_1+=     0.00000000133 * cos( 4.73063338736 +    18139.29450141590*t);
+   earth_y_1+=     0.00000000150 * cos( 4.37920284921 +     3340.61242669980*t);
+   earth_y_1+=     0.00000000125 * cos( 0.78764990758 +       20.77539549240*t);
+   earth_y_1+=     0.00000000123 * cos( 4.97498024339 +     6172.86952877200*t);
+   earth_y_1+=     0.00000000123 * cos( 4.81514598285 +     6393.28217121080*t);
+   earth_y_1+=     0.00000000135 * cos( 4.59213390099 +    10984.19235169980*t);
+   earth_y_1+=     0.00000000117 * cos( 2.90632080371 +    18852.75066832320*t);
+   earth_y_1+=     0.00000000118 * cos( 5.67650820416 +    17256.63153634140*t);
+   earth_y_1+=     0.00000000147 * cos( 4.60612631701 +    17654.78053974960*t);
+   earth_y_1+=     0.00000000113 * cos( 3.90726609167 +    65147.61976813770*t);
+   earth_y_1+=     0.00000000123 * cos( 2.85759760823 +    12565.17137891460*t);
+   earth_y_1+=     0.00000000107 * cos( 5.46942612515 +     9380.95967271720*t);
+   earth_y_1+=     0.00000000119 * cos( 0.71211331875 +    10419.98628350760*t);
+   earth_y_1+=     0.00000000105 * cos( 5.13254223053 +      632.78373931320*t);
+   earth_y_1+=     0.00000000103 * cos( 4.95941330270 +     6271.34649715540*t);
+   earth_y_1+=     0.00000000103 * cos( 4.83071292354 +     6294.80520282740*t);
+   earth_y_1+=     0.00000000102 * cos( 5.99360230997 +     4804.20927592700*t);
+   earth_y_1+=     0.00000000114 * cos( 3.78885344182 +     4171.42553661380*t);
+   earth_y_1+=     0.00000000113 * cos( 4.29988480427 +    23013.53953958720*t);
+   earth_y_1+=     0.00000000112 * cos( 0.61034141917 +    17253.04110768959*t);
+   earth_y_1+=     0.00000000111 * cos( 2.09128320599 +     1581.95934828300*t);
+   earth_y_1+=     0.00000000113 * cos( 0.97983808350 +    16858.48253293320*t);
+   earth_y_1+=     0.00000000102 * cos( 4.09962394648 +     6549.68289171320*t);
+   earth_y_1+=     0.00000000102 * cos( 5.69050227976 +     6016.46880826960*t);
+   earth_y_1+=     0.00000000106 * cos( 0.87687085688 +    23543.23050468179*t);
+   earth_y_1+=     0.00000000121 * cos( 4.76575782873 +    13625.53363017200*t);
+   earth_y_1+=     0.00000000090 * cos( 2.11774047359 +       76.26607127560*t);
+   earth_y_1+=     0.00000000088 * cos( 1.18967345782 +    17298.18232732620*t);
+   earth_y_1+=     0.00000000087 * cos( 1.12941525302 +    17267.26820169119*t);
+   earth_y_1+=     0.00000000092 * cos( 1.84158032509 +    22805.73556599360*t);
+   earth_y_1+=     0.00000000086 * cos( 4.36870119289 +    12808.88030395680*t);
+   earth_y_1+=     0.00000000111 * cos( 0.85102422912 +    11513.88331679440*t);
+   earth_y_1+=     0.00000000095 * cos( 6.05597457730 +    12345.73905754400*t);
+   earth_y_1+=     0.00000000106 * cos( 3.57199937061 +    11216.28429032400*t);
+   earth_y_1+=     0.00000000087 * cos( 1.36534411656 +     8982.81066930900*t);
+   earth_y_1+=     0.00000000095 * cos( 0.92966303480 +    12012.58229714040*t);
+   earth_y_1+=     0.00000000081 * cos( 2.27691876526 +    11614.43329373220*t);
+   earth_y_1+=     0.00000000084 * cos( 5.85994160066 +     6255.67429389460*t);
+   earth_y_1+=     0.00000000084 * cos( 3.93018462558 +     6310.47740608820*t);
+   earth_y_1+=     0.00000000080 * cos( 5.05950625443 +     6922.97313630540*t);
+   earth_y_1+=     0.00000000085 * cos( 1.08198764342 +     5636.06501667660*t);
+   earth_y_1+=     0.00000000084 * cos( 4.64040201203 +    10440.27429260360*t);
+   earth_y_1+=     0.00000000089 * cos( 6.25985265406 +    16723.35014259500*t);
+   earth_y_1+=     0.00000000080 * cos( 2.62622474570 +     3583.34103067380*t);
+   earth_y_1+=     0.00000000091 * cos( 5.59316682916 +    10818.13528691580*t);
+   earth_y_1+=     0.00000000091 * cos( 1.43612727660 +    15671.08175940660*t);
+   earth_y_1+=     0.00000000077 * cos( 4.13229801300 +     6993.00889854970*t);
+   earth_y_1+=     0.00000000076 * cos( 2.85409687552 +    11925.27409260060*t);
+   earth_y_1+=     0.00000000096 * cos( 1.47205677028 +     3097.88382272579*t);
+   earth_y_1+=     0.00000000104 * cos( 0.70746530546 +     5650.29211067820*t);
+   earth_y_1+=     0.00000000086 * cos( 5.52511956951 +    12359.96615154560*t);
+   earth_y_1+=     0.00000000080 * cos( 6.03205499783 +    13521.75144159140*t);
+   earth_y_1+=     0.00000000090 * cos( 4.45214765055 +     6195.76764545159*t);
+   earth_y_1+=     0.00000000090 * cos( 5.33797857569 +     6370.38405453120*t);
+   earth_y_1+=     0.00000000082 * cos( 4.91006846138 +    18842.11400297339*t);
+   earth_y_1+=     0.00000000073 * cos( 4.52976823745 +      949.17560896980*t);
+   earth_y_1+=     0.00000000083 * cos( 3.74911363362 +    18415.51581209739*t);
+   earth_y_1+=     0.00000000102 * cos( 4.43384315891 +     3104.93005942380*t);
+   earth_y_1+=     0.00000000074 * cos( 5.90055438198 +     5120.60114558360*t);
+   earth_y_1+=     0.00000000093 * cos( 2.53140640111 +    16460.33352952499*t);
+   earth_y_1+=     0.00000000084 * cos( 5.92514029201 +     2379.16447357160*t);
+   earth_y_1+=     0.00000000075 * cos( 0.05752343047 +    30355.99731976780*t);
+   earth_y_1+=     0.00000000076 * cos( 3.26399401386 +    18429.74290609900*t);
+   earth_y_1+=     0.00000000067 * cos( 4.00530001957 +    16062.18452611680*t);
+   earth_y_1+=     0.00000000066 * cos( 4.81634064130 +    10454.50138660520*t);
+   earth_y_1+=     0.00000000087 * cos( 3.52675041577 +     3903.91137641980*t);
+   earth_y_1+=     0.00000000078 * cos( 2.74560762552 +     8671.96987044060*t);
+   earth_y_1+=     0.00000000063 * cos( 5.70772167485 +    24356.78078864160*t);
+   earth_y_1+=     0.00000000073 * cos( 6.08379700055 +     9388.00590941520*t);
+   earth_y_1+=     0.00000000060 * cos( 1.94470153573 +    34570.06633485260*t);
+   earth_y_1+=     0.00000000057 * cos( 2.89018663881 +   149854.40013480789*t);
+   earth_y_1+=     0.00000000056 * cos( 3.76314176838 +    12455.94537876340*t);
+   earth_y_1+=     0.00000000055 * cos( 1.32127703933 +    23539.70738633280*t);
+   earth_y_1+=     0.00000000055 * cos( 5.60680941684 +    15664.03552270859*t);
+   earth_y_1+=     0.00000000056 * cos( 0.60887110989 +    11764.33076885900*t);
+   earth_y_1+=     0.00000000061 * cos( 4.86427999867 +    28237.23345938940*t);
+   earth_y_1+=     0.00000000065 * cos( 4.29965528550 +     7335.34423317980*t);
+   earth_y_1+=     0.00000000067 * cos( 4.20742811531 +    11933.36796066960*t);
+   earth_y_1+=     0.00000000072 * cos( 4.30216792995 +     6510.55198278040*t);
+   earth_y_1+=     0.00000000072 * cos( 5.48795829629 +     6055.59971720240*t);
+   earth_y_1+=     0.00000000053 * cos( 5.64466034451 +    17996.03116822220*t);
+   earth_y_1+=     0.00000000071 * cos( 4.94322519248 +    12545.37630449040*t);
+   earth_y_1+=     0.00000000061 * cos( 3.80625548858 +     6253.64634145540*t);
+   earth_y_1+=     0.00000000061 * cos( 5.98387073766 +     6312.50535852740*t);
+   earth_y_1+=     0.00000000052 * cos( 1.83796108842 +     5326.78669402080*t);
+   earth_y_1+=     0.00000000066 * cos( 0.79633834660 +    14919.01785375460*t);
+   earth_y_1+=     0.00000000053 * cos( 2.81005325604 +     2699.73481931760*t);
+   earth_y_1+=     0.00000000051 * cos( 2.33512659118 +    16522.65971600220*t);
+   earth_y_1+=     0.00000000049 * cos( 2.61656735675 +    35579.69123957000*t);
+   earth_y_1+=     0.00000000053 * cos( 3.71469727042 +     6357.85744855870*t);
+   earth_y_1+=     0.00000000053 * cos( 6.07542895582 +     6208.29425142410*t);
+   earth_y_1+=     0.00000000058 * cos( 6.24714677357 +    16737.57723659660*t);
+   earth_y_1+=     0.00000000049 * cos( 3.35591683551 +    21393.54196985760*t);
+   earth_y_1+=     0.00000000049 * cos( 5.31764383927 +     7018.95236352320*t);
+   earth_y_1+=     0.00000000049 * cos( 4.47248238697 +     5547.19933645960*t);
+   earth_y_1+=     0.00000000049 * cos( 1.02044068888 +    12416.58850284820*t);
+   earth_y_1+=     0.00000000048 * cos( 0.93318398067 +     5573.14280143310*t);
+   earth_y_1+=     0.00000000050 * cos( 4.78756360288 +    26084.02180621620*t);
+   earth_y_1+=     0.00000000052 * cos( 0.61376591765 +    45585.17281218740*t);
+   earth_y_1+=     0.00000000051 * cos( 1.83598363802 +    23006.42599258639*t);
+   earth_y_1+=     0.00000000058 * cos( 4.03624047121 +    15141.39079431200*t);
+   earth_y_1+=     0.00000000047 * cos( 3.55960074091 +    22743.40937951640*t);
+   earth_y_1+=     0.00000000050 * cos( 5.52194731109 +    21424.46664430340*t);
+   earth_y_1+=     0.00000000047 * cos( 5.68901702449 +    24279.10701821359*t);
+   earth_y_1+=     0.00000000049 * cos( 0.09767152059 +     5749.86176654780*t);
+   earth_y_1+=     0.00000000049 * cos( 3.40926939847 +     6816.28993343500*t);
+   earth_y_1+=     0.00000000046 * cos( 1.88054905125 +    19378.91851506880*t);
+   earth_y_1+=     0.00000000047 * cos( 0.52217551707 +    24065.80792277559*t);
+   earth_y_1+=     0.00000000046 * cos( 4.16716226231 +   104351.61256629678*t);
+   earth_y_1+=     0.00000000050 * cos( 2.92366135772 +    24602.61243487099*t);
+   earth_y_1+=     0.00000000054 * cos( 4.75214506546 +     7834.12107263940*t);
+   earth_y_1+=     0.00000000045 * cos( 3.48834172987 +   161000.68573767410*t);
+   earth_y_1+=     0.00000000049 * cos( 2.94403678329 +    29303.72893657940*t);
+   earth_y_1+=     0.00000000048 * cos( 1.95408707684 +     5244.04923920100*t);
+   earth_y_1+=     0.00000000043 * cos( 3.93857315871 +    12132.43996210600*t);
+   earth_y_1+=     0.00000000047 * cos( 5.81677586075 +    18208.34994259200*t);
+   earth_y_1+=     0.00000000049 * cos( 3.41896064981 +    29289.50184257780*t);
+   earth_y_1+=     0.00000000044 * cos( 2.06685104480 +     6175.41232605280*t);
+   earth_y_1+=     0.00000000044 * cos( 1.44008987426 +     6390.73937393000*t);
+   earth_y_1+=     0.00000000042 * cos( 6.18541909471 +    24705.70520908959*t);
+   earth_y_1+=     0.00000000041 * cos( 3.21613978391 +    17370.36097590980*t);
+   earth_y_1+=     0.00000000047 * cos( 1.17418881008 +    58953.14544329400*t);
+   earth_y_1+=     0.00000000041 * cos( 2.05518926769 +    23020.65308658799*t);
+   earth_y_1+=     0.00000000044 * cos( 3.20914273688 +    26735.94526221320*t);
+   earth_y_1+=     0.00000000045 * cos( 3.70228663504 +    54247.41313575040*t);
+   earth_y_1+=     0.00000000040 * cos( 4.57474127522 +    32243.49850088060*t);
+   earth_y_1+=     0.00000000040 * cos( 5.17182769558 +    24336.00539314920*t);
+   earth_y_1+=     0.00000000040 * cos( 3.62816477586 +    19804.82729158280*t);
+   earth_y_1+=     0.00000000039 * cos( 1.13462594128 +    18606.49894600020*t);
+   earth_y_1+=     0.00000000040 * cos( 4.71493835645 +     8401.83971036980*t);
+   earth_y_1+=     0.00000000041 * cos( 2.69655967341 +    16207.88627150200*t);
+   earth_y_1+=     0.00000000038 * cos( 5.56507429543 +    29088.81141598500*t);
+   earth_y_1+=     0.00000000038 * cos( 2.36481926067 +     6268.84875598980*t);
+   earth_y_1+=     0.00000000038 * cos( 1.14212165839 +     6297.30294399300*t);
+   earth_y_1+=     0.00000000041 * cos( 3.38234777816 +    15508.61512327440*t);
+   earth_y_1+=     0.00000000038 * cos( 4.82215691915 +     6171.88920770380*t);
+   earth_y_1+=     0.00000000038 * cos( 4.96796930709 +     6394.26249227900*t);
+   earth_y_1+=     0.00000000039 * cos( 1.75198804720 +    29826.30635467320*t);
+   earth_y_1+=     0.00000000037 * cos( 3.64776861935 +    20452.86941222180*t);
+   earth_y_1+=     0.00000000052 * cos( 3.98463781177 +     1047.74731175470*t);
+   earth_y_1+=     0.00000000043 * cos( 6.01711305973 +    55022.93574707440*t);
+   earth_y_1+=     0.00000000037 * cos( 6.18577055273 +     4061.21921539440*t);
+   earth_y_1+=     0.00000000040 * cos( 2.02818434731 +    49515.38250840700*t);
+   earth_y_1+=     0.00000000049 * cos( 4.21086629146 +    18216.44381066100*t);
+   earth_y_1+=     0.00000000038 * cos( 4.59185538964 +    18202.21671665939*t);
+   earth_y_1+=     0.00000000035 * cos( 3.16929714068 +    26482.17080962440*t);
+   earth_y_1+=     0.00000000036 * cos( 1.56812014048 +     9924.81042151060*t);
+   earth_y_1+=     0.00000000039 * cos( 0.88019393819 +    24889.57479599160*t);
+   earth_y_1+=     0.00000000039 * cos( 4.63870568395 +    26709.64694241340*t);
+   earth_y_1+=     0.00000000036 * cos( 6.07069720518 +    25287.72379939980*t);
+   earth_y_1+=     0.00000000034 * cos( 6.10039920596 +     6147.45052498140*t);
+   earth_y_1+=     0.00000000034 * cos( 3.68972702028 +     6418.70117500140*t);
+   earth_y_1+=     0.00000000034 * cos( 2.54760886649 +    21228.39202354580*t);
+   earth_y_1+=     0.00000000034 * cos( 0.68922440987 +    33019.02111220460*t);
+   earth_y_1+=     0.00000000034 * cos( 0.40887359374 +    17271.88400752640*t);
+   earth_y_1+=     0.00000000034 * cos( 0.33424154418 +    29026.48522950779*t);
+   earth_y_1+=     0.00000000035 * cos( 1.63582495302 +     7322.10246078180*t);
+   earth_y_1+=     0.00000000044 * cos( 1.21937055747 +     3981.49003408200*t);
+   earth_y_1+=     0.00000000033 * cos( 2.56048298726 +   173566.83743765688*t);
+   earth_y_1+=     0.00000000035 * cos( 3.00160924335 +    21954.15760939799*t);
+   earth_y_1+=     0.00000000034 * cos( 3.38052068808 +    38500.27603107220*t);
+   earth_y_1+=     0.00000000032 * cos( 2.81219711037 +    17686.75284556640*t);
+   earth_y_1+=     0.00000000033 * cos( 0.13837047156 +    34520.30930938080*t);
+   earth_y_1+=     0.00000000041 * cos( 2.52021125438 +    21202.09370374600*t);
+   earth_y_1+=     0.00000000037 * cos( 5.92836585991 +     7349.57132718140*t);
+   earth_y_1+=     0.00000000031 * cos( 2.79498926181 +    28230.18722269139*t);
+   earth_y_1+=     0.00000000038 * cos( 5.01277542891 +     9225.53927328300*t);
+   earth_y_1+=     0.00000000035 * cos( 2.49445060142 +     8584.66166590080*t);
+   earth_y_1+=     0.00000000038 * cos( 2.67041435749 +    15110.46611986620*t);
+   earth_y_1+=     0.00000000037 * cos( 1.84827551648 +    29296.61538957860*t);
+   earth_y_1+=     0.00000000035 * cos( 4.26342089717 +    11919.14086666800*t);
+   earth_y_1+=     0.00000000035 * cos( 2.59482262662 +    62883.35513951360*t);
+   earth_y_1+=     0.00000000035 * cos( 5.54313072444 +    29062.51309618520*t);
+   earth_y_1+=     0.00000000035 * cos( 5.96888200431 +    25685.87280280800*t);
+   earth_y_1+=     0.00000000030 * cos( 1.38932865248 +    12489.88562870720*t);
+   earth_y_1+=     0.00000000041 * cos( 4.39843263561 +    18052.92954315780*t);
+   earth_y_1+=     0.00000000032 * cos( 4.69085169487 +     6231.79551212899*t);
+   earth_y_1+=     0.00000000032 * cos( 5.09927453137 +     6334.35618785381*t);
+   earth_y_1+=     0.00000000033 * cos( 4.71041781828 +     5017.50837136500*t);
+   earth_y_1+=     0.00000000034 * cos( 3.82387178783 +    22483.84857449259*t);
+   earth_y_1+=     0.00000000029 * cos( 0.69173791628 +    27278.46881644080*t);
+   earth_y_1+=     0.00000000028 * cos( 0.30766379271 +    24080.03501677719*t);
+   earth_y_1+=     0.00000000028 * cos( 1.80200703459 +    21150.81336588360*t);
+   earth_y_1+=     0.00000000034 * cos( 5.49813240003 +     5905.70224207560*t);
+   earth_y_1+=     0.00000000034 * cos( 4.29199382621 +     6660.44945790720*t);
+   earth_y_1+=     0.00000000031 * cos( 0.77685387250 +    15265.88651930040*t);
+   earth_y_1+=     0.00000000038 * cos( 3.55141531291 +    31570.79964939120*t);
+   earth_y_1+=     0.00000000027 * cos( 2.03421614285 +    36949.23080842420*t);
+   earth_y_1+=     0.00000000030 * cos( 0.75326086687 +    12586.92709547520*t);
+   earth_y_1+=     0.00000000027 * cos( 5.24483718904 +   233141.31440436149*t);
+   earth_y_1+=     0.00000000027 * cos( 4.99665333874 +    24383.07910844140*t);
+   earth_y_1+=     0.00000000028 * cos( 3.97690125997 +    34911.41207609100*t);
+   earth_y_1+=     0.00000000027 * cos( 4.60109049263 +     1692.16566950240*t);
+   earth_y_1+=     0.00000000027 * cos( 3.01465510733 +    23581.25817731760*t);
+   earth_y_1+=     0.00000000028 * cos( 5.79946172825 +    16193.65917750039*t);
+   earth_y_1+=     0.00000000027 * cos( 5.36876316169 +     6360.74962041940*t);
+   earth_y_1+=     0.00000000027 * cos( 4.42136306455 +     6205.40207956340*t);
+   earth_y_1+=     0.00000000032 * cos( 5.04711925057 +    27511.46787353720*t);
+   earth_y_1+=     0.00000000027 * cos( 0.74966550512 +    32367.09765620760*t);
+   earth_y_1+=     0.00000000027 * cos( 3.34326768159 +     4487.81740627040*t);
+   earth_y_1+=     0.00000000030 * cos( 4.32132101590 +    52175.80628314840*t);
+   earth_y_1+=     0.00000000029 * cos( 4.00870617495 +    46386.99374331120*t);
+   earth_y_1+=     0.00000000026 * cos( 1.30285929374 +     4797.09572892620*t);
+   earth_y_1+=     0.00000000028 * cos( 5.37411984168 +    10239.58386601080*t);
+   earth_y_1+=     0.00000000026 * cos( 5.08151018860 +     6339.97422492700*t);
+   earth_y_1+=     0.00000000026 * cos( 4.70861603764 +     6226.17747505580*t);
+   earth_y_1+=     0.00000000029 * cos( 4.88985723155 +     6592.35417264720*t);
+   earth_y_1+=     0.00000000029 * cos( 4.90026899469 +     5973.79752733560*t);
+   earth_y_1+=     0.00000000026 * cos( 3.80920015468 +    25158.60171976540*t);
+   earth_y_1+=     0.00000000026 * cos( 4.12353757087 +     9910.58332750900*t);
+   earth_y_1+=     0.00000000031 * cos( 2.37034813214 +    24491.42579258340*t);
+   earth_y_1+=     0.00000000024 * cos( 5.26265995716 +    11616.97609101300*t);
+   earth_y_1+=     0.00000000024 * cos( 4.68388215645 +    31172.65064598300*t);
+   earth_y_1+=     0.00000000029 * cos( 4.62741692698 +    19402.79695281660*t);
+   earth_y_1+=     0.00000000025 * cos( 1.21411432675 +     6302.74561089119*t);
+   earth_y_1+=     0.00000000025 * cos( 2.29282659231 +     6263.40608909160*t);
+   earth_y_1+=     0.00000000026 * cos( 4.64443513632 +    30666.15495843280*t);
+   earth_y_1+=     0.00000000027 * cos( 2.79813911349 +    19800.94595622480*t);
+   earth_y_1+=     0.00000000024 * cos( 2.16373217847 +    31968.94865279940*t);
+   earth_y_1+=     0.00000000023 * cos( 5.33583152283 +    17892.93839400359*t);
+   earth_y_1+=     0.00000000023 * cos( 5.05383022333 +    29424.63423291600*t);
+   earth_y_1+=     0.00000000024 * cos( 6.12085409163 +    20597.24396304120*t);
+   earth_y_1+=     0.00000000024 * cos( 1.65155411828 +    23550.34405168259*t);
+   earth_y_1+=     0.00000000023 * cos( 6.07302432574 +     6716.78758786820*t);
+   earth_y_1+=     0.00000000022 * cos( 2.42662159283 +     7445.55055439920*t);
+   earth_y_1+=     0.00000000027 * cos( 1.56823801633 +    28628.33622609960*t);
+   earth_y_1+=     0.00000000022 * cos( 3.39570466574 +    40879.44050464380*t);
+   earth_y_1+=     0.00000000025 * cos( 5.52794317615 +    14155.22459526660*t);
+   earth_y_1+=     0.00000000023 * cos( 1.49204312505 +     6272.03014972750*t);
+   earth_y_1+=     0.00000000023 * cos( 2.01489779401 +     6294.12155025530*t);
+   earth_y_1+=     0.00000000021 * cos( 4.21170718841 +    22779.43724619380*t);
+   earth_y_1+=     0.00000000024 * cos( 3.83904364251 +     8186.51266249260*t);
+   earth_y_1+=     0.00000000021 * cos( 2.51494728863 +    30220.93223973240*t);
+   earth_y_1+=     0.00000000022 * cos( 6.09500922463 +     4379.63903749020*t);
+   earth_y_1+=     0.00000000022 * cos( 2.95144246643 +    28759.81087748319*t);
+   earth_y_1+=     0.00000000020 * cos( 0.07597154617 +    28313.28880466100*t);
+   earth_y_1+=     0.00000000026 * cos( 1.83971959567 +    18422.62935909819*t);
+   earth_y_1+=     0.00000000022 * cos( 1.06284971840 +    27707.54249429480*t);
+   earth_y_1+=     0.00000000027 * cos( 5.38238146057 +    34513.26307268280*t);
+   earth_y_1+=     0.00000000025 * cos( 1.96129048463 +    55798.45835839840*t);
+   earth_y_1+=     0.00000000019 * cos( 0.75805525916 +    11712.95531823080*t);
+   earth_y_1+=     0.00000000025 * cos( 4.06271601240 +    29864.33402730900*t);
+   earth_y_1+=     0.00000000024 * cos( 5.92086856803 +     9411.46461508720*t);
+   earth_y_1+=     0.00000000020 * cos( 0.08212819079 +     6379.05507720920*t);
+   earth_y_1+=     0.00000000020 * cos( 3.42481272827 +     6187.09662277360*t);
+   earth_y_1+=     0.00000000022 * cos( 2.87393551670 +    18326.65013188039*t);
+   earth_y_1+=     0.00000000020 * cos( 4.26784309163 +    21548.96236929180*t);
+   earth_y_1+=     0.00000000022 * cos( 2.89425510955 +     6332.83287546320*t);
+   earth_y_1+=     0.00000000022 * cos( 0.61268580951 +     6233.31882451960*t);
+   earth_y_1+=     0.00000000021 * cos( 5.43088404515 +    50317.20343953080*t);
+   earth_y_1+=     0.00000000020 * cos( 0.38470757997 +    22345.26037610820*t);
+   earth_y_1+=     0.00000000018 * cos( 3.04737429764 +     6593.91664885980*t);
+   earth_y_1+=     0.00000000018 * cos( 0.45956662142 +     5972.23505112300*t);
+   earth_y_1+=     0.00000000019 * cos( 0.79376120003 +     5618.31980486140*t);
+   earth_y_1+=     0.00000000019 * cos( 2.71317971903 +     6947.83189512140*t);
+   earth_y_1+=     0.00000000018 * cos( 0.28017726157 +    34115.11406927460*t);
+   earth_y_1+=     0.00000000019 * cos( 4.89700323392 +    83973.83535572988*t);
+   earth_y_1+=     0.00000000019 * cos( 2.72494815105 +    84019.85928049390*t);
+   earth_y_1+=     0.00000000020 * cos( 2.60251083721 +    42456.78404709160*t);
+   earth_y_1+=     0.00000000020 * cos( 1.28313240949 +     9070.11887384880*t);
+   earth_y_1=earth_y_1 * t;
+
+   double earth_y_2=0.0;
+   earth_y_2+=     0.00005080208 * cos( 0.00000000000 +        0.00000000000*t);
+   earth_y_2+=     0.00002178016 * cos( 2.82957544235 +    12566.15169998280*t);
+   earth_y_2+=     0.00001020487 * cos( 4.63746718598 +     6283.07584999140*t);
+   earth_y_2+=     0.00000092688 * cos( 2.62218748420 +    18849.22754997420*t);
+   earth_y_2+=     0.00000013680 * cos( 1.80736896884 +     6286.59896834040*t);
+   earth_y_2+=     0.00000013668 * cos( 1.70191859339 +     6279.55273164240*t);
+   earth_y_2+=     0.00000008190 * cos( 2.84848380902 +     6256.77753019160*t);
+   earth_y_2+=     0.00000008145 * cos( 0.65845469208 +     6309.37416979120*t);
+   earth_y_2+=     0.00000008004 * cos( 4.35208516133 +     6127.65545055720*t);
+   earth_y_2+=     0.00000007754 * cos( 5.43893832496 +     6438.49624942560*t);
+   earth_y_2+=     0.00000008900 * cos( 3.80944637329 +    83996.84731811189*t);
+   earth_y_2+=     0.00000004754 * cos( 4.59631649528 +      709.93304855830*t);
+   earth_y_2+=     0.00000002895 * cos( 3.37918924673 +     4705.73230754360*t);
+   earth_y_2+=     0.00000003730 * cos( 0.91975073803 +     5507.55323866740*t);
+   earth_y_2+=     0.00000002937 * cos( 2.51081180566 +    25132.30339996560*t);
+   earth_y_2+=     0.00000003203 * cos( 2.57193731532 +     7058.59846131540*t);
+   earth_y_2+=     0.00000002268 * cos( 1.94917686677 +      775.52261132400*t);
+   earth_y_2+=     0.00000002148 * cos( 0.14152116438 +     7860.41939243920*t);
+   earth_y_2+=     0.00000002031 * cos( 0.72283806291 +     6275.96230299060*t);
+   earth_y_2+=     0.00000002031 * cos( 2.78410285615 +     6290.18939699220*t);
+   earth_y_2+=     0.00000001904 * cos( 2.89589054598 +     5486.77784317500*t);
+   earth_y_2+=     0.00000001728 * cos( 3.66132960737 +     1059.38193018920*t);
+   earth_y_2+=     0.00000001553 * cos( 0.56386962806 +     6040.34724601740*t);
+   earth_y_2+=     0.00000001563 * cos( 0.61022183613 +     7079.37385680780*t);
+   earth_y_2+=     0.00000001459 * cos( 1.88257310296 +     5753.38488489680*t);
+   earth_y_2+=     0.00000001425 * cos( 1.59147859762 +     6812.76681508600*t);
+   earth_y_2+=     0.00000001436 * cos( 1.45357168378 +     5884.92684658320*t);
+   earth_y_2+=     0.00000001470 * cos( 2.94579999718 +     6525.80445396540*t);
+   earth_y_2+=     0.00000001288 * cos( 2.05812819149 +     6681.22485339960*t);
+   earth_y_2+=     0.00000001338 * cos( 5.75684444528 +     5729.50644714900*t);
+   earth_y_2+=     0.00000001039 * cos( 4.28230481184 +     6282.09552892320*t);
+   earth_y_2+=     0.00000001039 * cos( 5.50782141440 +     6284.05617105960*t);
+   earth_y_2+=     0.00000001191 * cos( 4.01917074242 +     6836.64525283380*t);
+   earth_y_2+=     0.00000000904 * cos( 0.86455323243 +     5331.35744374080*t);
+   earth_y_2+=     0.00000000871 * cos( 2.27688978691 +     4933.20844033260*t);
+   earth_y_2+=     0.00000000877 * cos( 1.49469829015 +    11790.62908865880*t);
+   earth_y_2+=     0.00000001102 * cos( 3.02451706631 +   167283.76158766549*t);
+   earth_y_2+=     0.00000000774 * cos( 2.68163901032 +     7234.79425624200*t);
+   earth_y_2+=     0.00000000746 * cos( 5.06505548485 +     4136.91043351620*t);
+   earth_y_2+=     0.00000000938 * cos( 2.05233105969 +    11856.21865142450*t);
+   earth_y_2+=     0.00000000653 * cos( 6.14159854835 +    11506.76976979360*t);
+   earth_y_2+=     0.00000000653 * cos( 1.20042991173 +     7632.94325965020*t);
+   earth_y_2+=     0.00000000855 * cos( 3.67241064264 +     4535.05943692440*t);
+   earth_y_2+=     0.00000000664 * cos( 3.21978126910 +     1194.44701022460*t);
+   earth_y_2+=     0.00000000814 * cos( 1.07520259172 +      155.42039943420*t);
+   earth_y_2+=     0.00000000579 * cos( 4.64690913221 +      796.29800681640*t);
+   earth_y_2+=     0.00000000778 * cos( 4.32748087995 +     5088.62883976680*t);
+   earth_y_2+=     0.00000000511 * cos( 0.71257437268 +     4694.00295470760*t);
+   earth_y_2+=     0.00000000686 * cos( 3.47131000180 +      801.82093112380*t);
+   earth_y_2+=     0.00000000677 * cos( 1.14891148781 +     3128.38876509580*t);
+   earth_y_2+=     0.00000000499 * cos( 1.45672603171 +    12569.67481833180*t);
+   earth_y_2+=     0.00000000480 * cos( 4.72566522648 +     8429.24126646660*t);
+   earth_y_2+=     0.00000000598 * cos( 6.12701256696 +     8031.09226305840*t);
+   earth_y_2+=     0.00000000499 * cos( 4.61877818070 +     1589.07289528380*t);
+   earth_y_2+=     0.00000000562 * cos( 5.48344520380 +     7477.52286021600*t);
+   earth_y_2+=     0.00000000415 * cos( 5.06370124318 +     6496.37494542940*t);
+   earth_y_2+=     0.00000000414 * cos( 4.73489535739 +     6069.77675455340*t);
+   earth_y_2+=     0.00000000404 * cos( 0.15585991312 +     3738.76143010800*t);
+   earth_y_2+=     0.00000000500 * cos( 1.83746339810 +     1592.59601363280*t);
+   earth_y_2+=     0.00000000392 * cos( 0.64094689719 +       26.29831979980*t);
+   earth_y_2+=     0.00000000366 * cos( 5.93255688855 +      553.56940284240*t);
+   earth_y_2+=     0.00000000458 * cos( 5.56847632382 +   156137.47598479928*t);
+   earth_y_2+=     0.00000000320 * cos( 0.43684109891 +     1990.74501704100*t);
+   earth_y_2+=     0.00000000318 * cos( 4.53872367492 +      951.71840625060*t);
+   earth_y_2+=     0.00000000361 * cos( 3.78139914914 +     7084.89678111520*t);
+   earth_y_2+=     0.00000000356 * cos( 4.78764157089 +    12721.57209941700*t);
+   earth_y_2+=     0.00000000279 * cos( 5.51531767271 +     4690.47983635860*t);
+   earth_y_2+=     0.00000000347 * cos( 6.05487619509 +     5481.25491886760*t);
+   earth_y_2+=     0.00000000279 * cos( 0.65072513623 +      242.72860397400*t);
+   earth_y_2+=     0.00000000367 * cos( 2.35508391326 +     9437.76293488700*t);
+   earth_y_2+=     0.00000000260 * cos( 0.98725902772 +     5223.69391980220*t);
+   earth_y_2+=     0.00000000255 * cos( 0.26797013701 +    11371.70468975820*t);
+   earth_y_2+=     0.00000000242 * cos( 3.35248006615 +     8827.39026987480*t);
+   earth_y_2+=     0.00000000245 * cos( 2.60270196816 +     3154.68708489560*t);
+   earth_y_2+=     0.00000000222 * cos( 0.17857635292 +     4164.31198961300*t);
+   earth_y_2+=     0.00000000221 * cos( 1.83683124614 +     1748.01641306700*t);
+   earth_y_2+=     0.00000000288 * cos( 3.16521452041 +     1349.86740965880*t);
+   earth_y_2+=     0.00000000204 * cos( 5.15479639183 +    11769.85369316640*t);
+   earth_y_2+=     0.00000000188 * cos( 4.08810200628 +     7875.67186362420*t);
+   earth_y_2+=     0.00000000204 * cos( 5.17246322251 +    10977.07880469900*t);
+   earth_y_2+=     0.00000000188 * cos( 5.32022679153 +     2388.89402044920*t);
+   earth_y_2+=     0.00000000183 * cos( 2.53020407005 +     7342.45778018060*t);
+   earth_y_2+=     0.00000000196 * cos( 4.81306031436 +     3930.20969621960*t);
+   earth_y_2+=     0.00000000207 * cos( 0.71423535949 +    13367.97263110660*t);
+   earth_y_2+=     0.00000000206 * cos( 1.68118727655 +    10973.55568635000*t);
+   earth_y_2+=     0.00000000144 * cos( 3.12119544517 +    10575.40668294180*t);
+   earth_y_2+=     0.00000000155 * cos( 4.24333182826 +    12410.73130054860*t);
+   earth_y_2+=     0.00000000135 * cos( 5.62460495351 +    12592.45001978260*t);
+   earth_y_2+=     0.00000000142 * cos( 3.43769023959 +      529.69096509460*t);
+   earth_y_2+=     0.00000000125 * cos( 2.30734650656 +    12562.62858163380*t);
+   earth_y_2+=     0.00000000118 * cos( 3.15911802638 +    13119.72110282519*t);
+   earth_y_2+=     0.00000000124 * cos( 4.88716980692 +     8635.94200376320*t);
+   earth_y_2+=     0.00000000111 * cos( 0.80973565731 +      398.14900340820*t);
+   earth_y_2+=     0.00000000099 * cos( 1.74904579789 +    13517.87010623340*t);
+   earth_y_2+=     0.00000000097 * cos( 4.08519883679 +    17260.15465469040*t);
+   earth_y_2+=     0.00000000102 * cos( 2.80042792551 +    10988.80815753500*t);
+   earth_y_2+=     0.00000000093 * cos( 4.87729654087 +     1577.34354244780*t);
+   earth_y_2+=     0.00000000094 * cos( 4.56913615419 +    10177.25767953360*t);
+   earth_y_2+=     0.00000000095 * cos( 0.58410858338 +     2146.16541647520*t);
+   earth_y_2+=     0.00000000100 * cos( 0.77889460323 +    13095.84266507740*t);
+   earth_y_2+=     0.00000000082 * cos( 3.65082093235 +    16730.46368959580*t);
+   earth_y_2+=     0.00000000109 * cos( 2.42647120751 +    12539.85338018300*t);
+   earth_y_2+=     0.00000000080 * cos( 2.67636202935 +    12323.42309600880*t);
+   earth_y_2+=     0.00000000105 * cos( 1.96765666886 +     5642.19824260920*t);
+   earth_y_2+=     0.00000000101 * cos( 0.26697958405 +     4292.33083295040*t);
+   earth_y_2+=     0.00000000075 * cos( 4.15987686546 +     6709.67404086740*t);
+   earth_y_2+=     0.00000000075 * cos( 5.63024936078 +     5856.47765911540*t);
+   earth_y_2+=     0.00000000100 * cos( 5.37997214401 +   250570.67585721909*t);
+   earth_y_2+=     0.00000000071 * cos( 5.31207567706 +    14314.16811304980*t);
+   earth_y_2+=     0.00000000092 * cos( 0.41159879070 +    13916.01910964160*t);
+   earth_y_2+=     0.00000000091 * cos( 1.53881053794 +     6923.95345737360*t);
+   earth_y_2+=     0.00000000082 * cos( 2.43138857903 +    31415.37924995700*t);
+   earth_y_2+=     0.00000000083 * cos( 4.40764560915 +    95143.13292097810*t);
+   earth_y_2+=     0.00000000061 * cos( 2.66905862210 +     8273.82086703240*t);
+   earth_y_2+=     0.00000000055 * cos( 2.98440193981 +     5643.17856367740*t);
+   earth_y_2+=     0.00000000056 * cos( 3.14343584528 +      639.89728631400*t);
+   earth_y_2+=     0.00000000054 * cos( 4.06071028317 +    16496.36139620240*t);
+   earth_y_2+=     0.00000000055 * cos( 2.26469062405 +    13341.67431130680*t);
+   earth_y_2+=     0.00000000073 * cos( 0.30591075918 +    71430.69561812909*t);
+   earth_y_2+=     0.00000000052 * cos( 5.97085689483 +     3634.62102451840*t);
+   earth_y_2+=     0.00000000050 * cos( 2.61636629528 +    12573.26524698360*t);
+   earth_y_2+=     0.00000000052 * cos( 6.19672300540 +    15720.83878487840*t);
+   earth_y_2+=     0.00000000065 * cos( 4.92633618792 +     4732.03062734340*t);
+   earth_y_2+=     0.00000000066 * cos( 3.92505637420 +     2787.04302385740*t);
+   earth_y_2+=     0.00000000048 * cos( 0.45273646609 +    13362.44970679920*t);
+   earth_y_2+=     0.00000000053 * cos( 2.55293773264 +    14919.01785375460*t);
+   earth_y_2+=     0.00000000047 * cos( 2.25879651839 +    12964.30070339100*t);
+   earth_y_2+=     0.00000000047 * cos( 1.06895129610 +     3894.18182954220*t);
+   earth_y_2+=     0.00000000042 * cos( 5.09300549015 +    24072.92146977640*t);
+   earth_y_2+=     0.00000000051 * cos( 5.97812516445 +     6195.76764545159*t);
+   earth_y_2+=     0.00000000051 * cos( 3.81200106179 +     6370.38405453120*t);
+   earth_y_2+=     0.00000000048 * cos( 5.31227829697 +     2118.76386037840*t);
+   earth_y_2+=     0.00000000041 * cos( 0.92655469135 +    18319.53658487960*t);
+   earth_y_2+=     0.00000000041 * cos( 2.52247946610 +     3185.19202726560*t);
+   earth_y_2+=     0.00000000056 * cos( 2.58677293772 +    84672.47584450469*t);
+   earth_y_2+=     0.00000000039 * cos( 2.70559904392 +    20426.57109242200*t);
+   earth_y_2+=     0.00000000053 * cos( 1.64080385432 +   239424.39025435288*t);
+   earth_y_2+=     0.00000000045 * cos( 6.04609930388 +     9779.10867612540*t);
+   earth_y_2+=     0.00000000044 * cos( 3.24184126256 +     8662.24032356300*t);
+   earth_y_2+=     0.00000000043 * cos( 2.12627080359 +    18139.29450141590*t);
+   earth_y_2+=     0.00000000033 * cos( 1.67764931937 +     3496.03282613400*t);
+   earth_y_2+=     0.00000000035 * cos( 5.17885502556 +    18875.52586977400*t);
+   earth_y_2+=     0.00000000035 * cos( 4.51995137508 +    10419.98628350760*t);
+   earth_y_2+=     0.00000000031 * cos( 5.35886426810 +    23013.53953958720*t);
+   earth_y_2+=     0.00000000030 * cos( 2.82636067861 +    12808.88030395680*t);
+   earth_y_2+=     0.00000000042 * cos( 0.90184598337 +    82576.98122099529*t);
+   earth_y_2+=     0.00000000031 * cos( 3.95971884541 +    14712.31711645800*t);
+   earth_y_2+=     0.00000000040 * cos( 4.02926744005 +    17789.84561978500*t);
+   earth_y_2+=     0.00000000031 * cos( 5.47911308336 +    22779.43724619380*t);
+   earth_y_2+=     0.00000000030 * cos( 2.31682203487 +    10447.38783960440*t);
+   earth_y_2+=     0.00000000028 * cos( 2.11876435927 +    14143.49524243060*t);
+   earth_y_2+=     0.00000000028 * cos( 3.91478954977 +     5429.87946823940*t);
+   earth_y_2+=     0.00000000030 * cos( 1.17555913144 +     9380.95967271720*t);
+   earth_y_2+=     0.00000000026 * cos( 0.72314347990 +     9623.68827669120*t);
+   earth_y_2+=     0.00000000026 * cos( 2.78379743916 +     2942.46342329160*t);
+   earth_y_2+=     0.00000000030 * cos( 1.83762487734 +    10213.28554621100*t);
+   earth_y_2+=     0.00000000026 * cos( 0.18151529861 +    18209.33026366019*t);
+   earth_y_2+=     0.00000000025 * cos( 1.11733039286 +     3583.34103067380*t);
+   earth_y_2+=     0.00000000024 * cos( 1.36775216012 +     9917.69687450980*t);
+   earth_y_2+=     0.00000000026 * cos( 5.61732376806 +    12012.58229714040*t);
+   earth_y_2+=     0.00000000023 * cos( 0.92088373151 +    12559.03815298200*t);
+   earth_y_2+=     0.00000000023 * cos( 5.52325028212 +    11015.10647733480*t);
+   earth_y_2+=     0.00000000023 * cos( 0.69647593831 +    26709.64694241340*t);
+   earth_y_2+=     0.00000000024 * cos( 3.30887701069 +    77713.77146812050*t);
+   earth_y_2+=     0.00000000022 * cos( 4.42086856143 +     6993.00889854970*t);
+   earth_y_2+=     0.00000000021 * cos( 0.81252266714 +    17256.63153634140*t);
+   earth_y_2+=     0.00000000023 * cos( 5.52062840479 +   149854.40013480789*t);
+   earth_y_2+=     0.00000000019 * cos( 0.35979025769 +    11926.25441366880*t);
+   earth_y_2+=     0.00000000018 * cos( 0.71176167314 +    11614.43329373220*t);
+   earth_y_2+=     0.00000000019 * cos( 4.89783197147 +    18073.70493865020*t);
+   earth_y_2+=     0.00000000019 * cos( 2.09021029647 +    30639.85663863300*t);
+   earth_y_2+=     0.00000000022 * cos( 2.19299415984 +    17298.18232732620*t);
+   earth_y_2+=     0.00000000019 * cos( 1.23251732918 +    19651.04848109800*t);
+   earth_y_2+=     0.00000000023 * cos( 3.21303801707 +    72850.56171524568*t);
+   earth_y_2+=     0.00000000017 * cos( 2.38624204004 +     8982.81066930900*t);
+   earth_y_2+=     0.00000000021 * cos( 2.20285403235 +    16858.48253293320*t);
+   earth_y_2+=     0.00000000021 * cos( 2.89968186701 +    12036.46073488820*t);
+   earth_y_2+=     0.00000000014 * cos( 3.30566599775 +    16460.33352952499*t);
+   earth_y_2+=     0.00000000020 * cos( 4.99490943445 +     5573.14280143310*t);
+   earth_y_2+=     0.00000000015 * cos( 2.64766821618 +    20199.09495963300*t);
+   earth_y_2+=     0.00000000014 * cos( 4.87181144278 +    22483.84857449259*t);
+   earth_y_2+=     0.00000000019 * cos( 1.24341688060 +     8671.96987044060*t);
+   earth_y_2+=     0.00000000014 * cos( 5.97526530111 +    90279.92316810328*t);
+   earth_y_2+=     0.00000000016 * cos( 5.12394281212 +    23543.23050468179*t);
+   earth_y_2+=     0.00000000012 * cos( 1.13908210541 +    17996.03116822220*t);
+   earth_y_2+=     0.00000000012 * cos( 5.43458827343 +    19402.79695281660*t);
+   earth_y_2+=     0.00000000012 * cos( 2.80798050151 +    10818.13528691580*t);
+   earth_y_2+=     0.00000000015 * cos( 5.28864619208 +    12567.13202105100*t);
+   earth_y_2+=     0.00000000011 * cos( 0.21390152319 +    28766.92442448400*t);
+   earth_y_2+=     0.00000000011 * cos( 0.68624711981 +    14945.31617355440*t);
+   earth_y_2+=     0.00000000011 * cos( 3.36248263984 +    18451.07854656599*t);
+   earth_y_2+=     0.00000000011 * cos( 4.69826436697 +     2352.86615377180*t);
+   earth_y_2+=     0.00000000011 * cos( 0.30993049479 +    22805.73556599360*t);
+   earth_y_2+=     0.00000000011 * cos( 0.67524830165 +    35371.88726597640*t);
+   earth_y_2+=     0.00000000014 * cos( 0.36237106347 +    21228.39202354580*t);
+   earth_y_2+=     0.00000000010 * cos( 4.48227917427 +    23141.55838292460*t);
+   earth_y_2+=     0.00000000010 * cos( 0.98562604234 +    22345.26037610820*t);
+   earth_y_2+=     0.00000000009 * cos( 6.20387541356 +    29826.30635467320*t);
+   earth_y_2+=     0.00000000010 * cos( 4.89597714031 +    12779.45079542080*t);
+   earth_y_2+=     0.00000000010 * cos( 4.53062887484 +      213.29909543800*t);
+   earth_y_2+=     0.00000000011 * cos( 5.10892321991 +    25158.60171976540*t);
+   earth_y_2+=     0.00000000009 * cos( 2.89659186256 +    23581.25817731760*t);
+   earth_y_2=earth_y_2 * t * t;
+
+   double earth_y_3=0.0;
+   earth_y_3+=     0.00000128116 * cos( 3.14159265359 +        0.00000000000*t);
+   earth_y_3+=     0.00000072366 * cos( 1.32019314413 +    12566.15169998280*t);
+   earth_y_3+=     0.00000013847 * cos( 5.37444701633 +     6283.07584999140*t);
+   earth_y_3+=     0.00000005044 * cos( 0.96001931081 +    18849.22754997420*t);
+   earth_y_3+=     0.00000001481 * cos( 0.66592602555 +     6438.49624942560*t);
+   earth_y_3+=     0.00000001481 * cos( 2.84101489351 +     6127.65545055720*t);
+   earth_y_3+=     0.00000000644 * cos( 3.31488272100 +     6279.55273164240*t);
+   earth_y_3+=     0.00000000644 * cos( 0.19205819806 +     6286.59896834040*t);
+   earth_y_3+=     0.00000000318 * cos( 2.06760677503 +     6040.34724601740*t);
+   earth_y_3+=     0.00000000318 * cos( 1.43933414403 +     6525.80445396540*t);
+   earth_y_3+=     0.00000000220 * cos( 0.77087312633 +    25132.30339996560*t);
+   earth_y_3+=     0.00000000201 * cos( 5.53953658848 +     6836.64525283380*t);
+   earth_y_3+=     0.00000000201 * cos( 4.25058963776 +     5729.50644714900*t);
+   earth_y_3+=     0.00000000036 * cos( 6.12441668493 +    12569.67481833180*t);
+   earth_y_3+=     0.00000000033 * cos( 2.66474952932 +    12410.73130054860*t);
+   earth_y_3+=     0.00000000018 * cos( 0.75629855075 +    12721.57209941700*t);
+   earth_y_3+=     0.00000000011 * cos( 4.70234349211 +    12592.45001978260*t);
+   earth_y_3=earth_y_3 * t * t * t;
+
+   double earth_y_4=0.0;
+   earth_y_4+=     0.00000004187 * cos( 3.14159265359 +        0.00000000000*t);
+   earth_y_4+=     0.00000001931 * cos( 5.89883838000 +    12566.15169998280*t);
+   earth_y_4+=     0.00000000595 * cos( 0.51679590840 +     6283.07584999140*t);
+   earth_y_4+=     0.00000000210 * cos( 2.17239897368 +     6438.49624942560*t);
+   earth_y_4+=     0.00000000210 * cos( 1.33454194538 +     6127.65545055720*t);
+   earth_y_4+=     0.00000000229 * cos( 5.59311519226 +    18849.22754997420*t);
+   earth_y_4+=     0.00000000044 * cos( 3.02550685011 +      709.93304855830*t);
+   earth_y_4+=     0.00000000021 * cos( 4.89633014340 +     6279.55273164240*t);
+   earth_y_4+=     0.00000000021 * cos( 4.89379608284 +     6286.59896834040*t);
+   earth_y_4+=     0.00000000013 * cos( 2.23957099685 +    83996.84731811189*t);
+   earth_y_4+=     0.00000000011 * cos( 5.36791829228 +    25132.30339996560*t);
+   earth_y_4+=     0.00000000009 * cos( 3.62304265651 +    11856.21865142450*t);
+   earth_y_4+=     0.00000000007 * cos( 4.59462394746 +   167283.76158766549*t);
+   earth_y_4+=     0.00000000005 * cos( 1.11835906029 +    12410.73130054860*t);
+   earth_y_4+=     0.00000000006 * cos( 3.99789149114 +   156137.47598479928*t);
+   earth_y_4=earth_y_4 * t * t * t * t;
+
+   double earth_y_5=0.0;
+   earth_y_5+=     0.00000000072 * cos( 0.00000000000 +        0.00000000000*t);
+   earth_y_5+=     0.00000000043 * cos( 4.51710019379 +    12566.15169998280*t);
+   earth_y_5+=     0.00000000025 * cos( 3.76700032666 +     6438.49624942560*t);
+   earth_y_5+=     0.00000000025 * cos( 6.02312589958 +     6127.65545055720*t);
+   earth_y_5+=     0.00000000021 * cos( 2.56882591004 +     6283.07584999140*t);
+   earth_y_5+=     0.00000000008 * cos( 3.95508428463 +    18849.22754997420*t);
+   earth_y_5=earth_y_5 * t * t * t * t * t;
+
+   return earth_y_0+earth_y_1+earth_y_2+earth_y_3+earth_y_4+earth_y_5;
+}
+
+double vsop87a_full::earth_z(double t){
+   double earth_z_0=0.0;
+   earth_z_0+=     0.00000279620 * cos( 3.19870156017 +    84334.66158130829*t);
+   earth_z_0+=     0.00000101625 * cos( 5.42248110597 +     5507.55323866740*t);
+   earth_z_0+=     0.00000080461 * cos( 3.88027157914 +     5223.69391980220*t);
+   earth_z_0+=     0.00000043855 * cos( 3.70430347723 +     2352.86615377180*t);
+   earth_z_0+=     0.00000031853 * cos( 3.99997840986 +     1577.34354244780*t);
+   earth_z_0+=     0.00000022724 * cos( 3.98473831560 +     1047.74731175470*t);
+   earth_z_0+=     0.00000016395 * cos( 3.56456162523 +     5856.47765911540*t);
+   earth_z_0+=     0.00000018150 * cos( 4.98479613321 +     6283.07584999140*t);
+   earth_z_0+=     0.00000014477 * cos( 3.70258423465 +     9437.76293488700*t);
+   earth_z_0+=     0.00000014325 * cos( 3.41020246931 +    10213.28554621100*t);
+   earth_z_0+=     0.00000011204 * cos( 4.83021499098 +    14143.49524243060*t);
+   earth_z_0+=     0.00000010896 * cos( 2.08023708321 +     6812.76681508600*t);
+   earth_z_0+=     0.00000009716 * cos( 3.47560205064 +     4694.00295470760*t);
+   earth_z_0+=     0.00000010363 * cos( 4.05663928093 +    71092.88135493269*t);
+   earth_z_0+=     0.00000008792 * cos( 4.44890729280 +     5753.38488489680*t);
+   earth_z_0+=     0.00000008337 * cos( 4.99167706048 +     7084.89678111520*t);
+   earth_z_0+=     0.00000006921 * cos( 4.32559054073 +     6275.96230299060*t);
+   earth_z_0+=     0.00000009145 * cos( 1.14182646613 +     6620.89011318780*t);
+   earth_z_0+=     0.00000007205 * cos( 3.62441599378 +      529.69096509460*t);
+   earth_z_0+=     0.00000007698 * cos( 5.55425745881 +   167621.57585086189*t);
+   earth_z_0+=     0.00000005186 * cos( 6.25384533126 +    18073.70493865020*t);
+   earth_z_0+=     0.00000005027 * cos( 2.49727910749 +     4705.73230754360*t);
+   earth_z_0+=     0.00000004529 * cos( 2.33827747356 +     6309.37416979120*t);
+   earth_z_0+=     0.00000004754 * cos( 0.71100829534 +     5884.92684658320*t);
+   earth_z_0+=     0.00000004292 * cos( 1.10034019036 +     6681.22485339960*t);
+   earth_z_0+=     0.00000003855 * cos( 1.82331477760 +     5486.77784317500*t);
+   earth_z_0+=     0.00000005256 * cos( 4.42445744523 +     7860.41939243920*t);
+   earth_z_0+=     0.00000004025 * cos( 5.11990285618 +    13367.97263110660*t);
+   earth_z_0+=     0.00000004061 * cos( 6.02923989849 +     3930.20969621960*t);
+   earth_z_0+=     0.00000003798 * cos( 0.44370219832 +     3154.68708489560*t);
+   earth_z_0+=     0.00000002933 * cos( 5.12428135363 +     1059.38193018920*t);
+   earth_z_0+=     0.00000003424 * cos( 5.42065191700 +     6069.77675455340*t);
+   earth_z_0+=     0.00000003670 * cos( 4.58210192227 +    12194.03291462090*t);
+   earth_z_0+=     0.00000002912 * cos( 1.92688087702 +    10977.07880469900*t);
+   earth_z_0+=     0.00000002506 * cos( 0.60484952637 +     6496.37494542940*t);
+   earth_z_0+=     0.00000002680 * cos( 1.39497359287 +    22003.91463486980*t);
+   earth_z_0+=     0.00000002120 * cos( 4.30691000285 +     5643.17856367740*t);
+   earth_z_0+=     0.00000002235 * cos( 0.81363184041 +     8635.94200376320*t);
+   earth_z_0+=     0.00000001819 * cos( 3.75748003107 +     3340.61242669980*t);
+   earth_z_0+=     0.00000002230 * cos( 2.76407822446 +    12036.46073488820*t);
+   earth_z_0+=     0.00000001945 * cos( 5.69949789177 +    11790.62908865880*t);
+   earth_z_0+=     0.00000001529 * cos( 1.98639348122 +      398.14900340820*t);
+   earth_z_0+=     0.00000001581 * cos( 3.19976230948 +     5088.62883976680*t);
+   earth_z_0+=     0.00000001424 * cos( 6.26116472313 +     2544.31441988340*t);
+   earth_z_0+=     0.00000001405 * cos( 4.69394873481 +     7058.59846131540*t);
+   earth_z_0+=     0.00000001588 * cos( 0.25704624784 +    17298.18232732620*t);
+   earth_z_0+=     0.00000001478 * cos( 2.81808207569 +    25934.12433108940*t);
+   earth_z_0+=     0.00000001693 * cos( 4.95689385293 +   156475.29024799570*t);
+   earth_z_0+=     0.00000001123 * cos( 2.38605285936 +     3738.76143010800*t);
+   earth_z_0+=     0.00000000990 * cos( 4.29600699156 +     9225.53927328300*t);
+   earth_z_0+=     0.00000000923 * cos( 3.07195736431 +     4164.31198961300*t);
+   earth_z_0+=     0.00000000863 * cos( 0.55119493097 +     8429.24126646660*t);
+   earth_z_0+=     0.00000000984 * cos( 5.96784225392 +     7079.37385680780*t);
+   earth_z_0+=     0.00000001086 * cos( 3.14159265359 +        0.00000000000*t);
+   earth_z_0+=     0.00000000824 * cos( 1.50050055147 +    10447.38783960440*t);
+   earth_z_0+=     0.00000000829 * cos( 3.51336978167 +    11506.76976979360*t);
+   earth_z_0+=     0.00000000908 * cos( 0.12102885474 +    11015.10647733480*t);
+   earth_z_0+=     0.00000000742 * cos( 1.99159139281 +    26087.90314157420*t);
+   earth_z_0+=     0.00000000850 * cos( 4.24120016095 +    29864.33402730900*t);
+   earth_z_0+=     0.00000000762 * cos( 2.90028034892 +     4732.03062734340*t);
+   earth_z_0+=     0.00000000711 * cos( 1.38596151724 +     2146.16541647520*t);
+   earth_z_0+=     0.00000000766 * cos( 0.61397570576 +      796.29800681640*t);
+   earth_z_0+=     0.00000000708 * cos( 1.91406542362 +     8031.09226305840*t);
+   earth_z_0+=     0.00000000810 * cos( 5.13961498599 +     2942.46342329160*t);
+   earth_z_0+=     0.00000000751 * cos( 1.67479850166 +    21228.39202354580*t);
+   earth_z_0+=     0.00000000589 * cos( 2.02193316047 +      775.52261132400*t);
+   earth_z_0+=     0.00000000513 * cos( 2.60734651128 +    12566.15169998280*t);
+   earth_z_0+=     0.00000000550 * cos( 1.61119298963 +      801.82093112380*t);
+   earth_z_0+=     0.00000000592 * cos( 4.59481504319 +     4690.47983635860*t);
+   earth_z_0+=     0.00000000519 * cos( 5.72925039114 +     8827.39026987480*t);
+   earth_z_0+=     0.00000000515 * cos( 4.09976487552 +    64809.80550494129*t);
+   earth_z_0+=     0.00000000503 * cos( 5.66433137112 +    33794.54372352860*t);
+   earth_z_0+=     0.00000000476 * cos( 3.69640472448 +      213.29909543800*t);
+   earth_z_0+=     0.00000000463 * cos( 1.24987240483 +    15720.83878487840*t);
+   earth_z_0+=     0.00000000411 * cos( 0.10769444525 +     3128.38876509580*t);
+   earth_z_0+=     0.00000000504 * cos( 3.26207669160 +     7632.94325965020*t);
+   earth_z_0+=     0.00000000552 * cos( 1.02926440457 +   239762.20451754928*t);
+   earth_z_0+=     0.00000000394 * cos( 5.81496021156 +      426.59819087600*t);
+   earth_z_0+=     0.00000000476 * cos( 3.52810472421 +    16496.36139620240*t);
+   earth_z_0+=     0.00000000503 * cos( 4.85802444134 +     6290.18939699220*t);
+   earth_z_0+=     0.00000000417 * cos( 0.81920713533 +     5216.58037280140*t);
+   earth_z_0+=     0.00000000363 * cos( 5.70015720452 +     1589.07289528380*t);
+   earth_z_0+=     0.00000000363 * cos( 1.28376436579 +     6206.80977871580*t);
+   earth_z_0+=     0.00000000353 * cos( 4.70059133110 +     7234.79425624200*t);
+   earth_z_0+=     0.00000000381 * cos( 2.61890749829 +     7342.45778018060*t);
+   earth_z_0+=     0.00000000420 * cos( 0.97821135132 +     4136.91043351620*t);
+   earth_z_0+=     0.00000000335 * cos( 6.22057967320 +    12168.00269657460*t);
+   earth_z_0+=     0.00000000387 * cos( 3.09145061418 +    25158.60171976540*t);
+   earth_z_0+=     0.00000000349 * cos( 2.94840720073 +     9623.68827669120*t);
+   earth_z_0+=     0.00000000418 * cos( 3.75759994446 +     5230.80746680300*t);
+   earth_z_0+=     0.00000000396 * cos( 1.22507712354 +     6438.49624942560*t);
+   earth_z_0+=     0.00000000322 * cos( 1.21162178805 +     8662.24032356300*t);
+   earth_z_0+=     0.00000000323 * cos( 5.50254808471 +     1194.44701022460*t);
+   earth_z_0+=     0.00000000372 * cos( 1.72256107938 +    14945.31617355440*t);
+   earth_z_0+=     0.00000000305 * cos( 0.80429352049 +    37724.75341974820*t);
+   earth_z_0+=     0.00000000265 * cos( 6.10358507671 +     6836.64525283380*t);
+   earth_z_0+=     0.00000000250 * cos( 4.56452895547 +     7477.52286021600*t);
+   earth_z_0+=     0.00000000266 * cos( 2.62926282354 +     7238.67559160000*t);
+   earth_z_0+=     0.00000000262 * cos( 1.34297269105 +    11769.85369316640*t);
+   earth_z_0+=     0.00000000263 * cos( 6.22089501237 +     6133.51265285680*t);
+   earth_z_0+=     0.00000000302 * cos( 2.81139045072 +     1748.01641306700*t);
+   earth_z_0+=     0.00000000316 * cos( 1.62662805006 +   250908.49012041549*t);
+   earth_z_0+=     0.00000000226 * cos( 2.42346415873 +    11371.70468975820*t);
+   earth_z_0+=     0.00000000216 * cos( 3.68721275185 +     5849.36411211460*t);
+   earth_z_0+=     0.00000000276 * cos( 2.33399767923 +    19651.04848109800*t);
+   earth_z_0+=     0.00000000230 * cos( 0.36165162947 +     5863.59120611620*t);
+   earth_z_0+=     0.00000000200 * cos( 5.86073159059 +     4535.05943692440*t);
+   earth_z_0+=     0.00000000277 * cos( 4.65400292395 +    82239.16695779889*t);
+   earth_z_0+=     0.00000000199 * cos( 5.05186622555 +     5429.87946823940*t);
+   earth_z_0+=     0.00000000206 * cos( 3.70556027982 +    10973.55568635000*t);
+   earth_z_0+=     0.00000000210 * cos( 4.50691909144 +    29088.81141598500*t);
+   earth_z_0+=     0.00000000181 * cos( 6.00294783127 +     4292.33083295040*t);
+   earth_z_0+=     0.00000000249 * cos( 0.12900984422 +   154379.79562448629*t);
+   earth_z_0+=     0.00000000167 * cos( 6.27348370813 +    10988.80815753500*t);
+   earth_z_0+=     0.00000000183 * cos( 3.19836584401 +    16730.46368959580*t);
+   earth_z_0+=     0.00000000225 * cos( 3.18339652605 +    18875.52586977400*t);
+   earth_z_0+=     0.00000000192 * cos( 0.87192032514 +    77375.95720492408*t);
+   earth_z_0+=     0.00000000188 * cos( 2.22746128596 +    41654.96311596780*t);
+   earth_z_0+=     0.00000000160 * cos( 5.20207795189 +     5481.25491886760*t);
+   earth_z_0+=     0.00000000192 * cos( 3.89678943865 +    17789.84561978500*t);
+   earth_z_0+=     0.00000000159 * cos( 5.32590846475 +        7.11354700080*t);
+   earth_z_0+=     0.00000000214 * cos( 4.08171403539 +      337.81426319640*t);
+   earth_z_0+=     0.00000000184 * cos( 3.75465072369 +    23581.25817731760*t);
+   earth_z_0+=     0.00000000160 * cos( 4.53897529970 +    18477.10876461230*t);
+   earth_z_0+=     0.00000000144 * cos( 5.19486374726 +    20426.57109242200*t);
+   earth_z_0+=     0.00000000126 * cos( 1.69446958184 +    15110.46611986620*t);
+   earth_z_0+=     0.00000000149 * cos( 1.55306832217 +    10021.83728009940*t);
+   earth_z_0+=     0.00000000127 * cos( 5.31068233515 +      639.89728631400*t);
+   earth_z_0+=     0.00000000118 * cos( 2.73268114062 +    18849.22754997420*t);
+   earth_z_0+=     0.00000000123 * cos( 2.55361087146 +    12592.45001978260*t);
+   earth_z_0+=     0.00000000113 * cos( 6.07178776726 +     6709.67404086740*t);
+   earth_z_0+=     0.00000000130 * cos( 3.94161517411 +     7330.82316174610*t);
+   earth_z_0+=     0.00000000110 * cos( 3.51203379263 +    18052.92954315780*t);
+   earth_z_0+=     0.00000000147 * cos( 4.63371971408 +    22805.73556599360*t);
+   earth_z_0+=     0.00000000113 * cos( 1.20711498905 +    14919.01785375460*t);
+   earth_z_0+=     0.00000000108 * cos( 3.76359880753 +    14314.16811304980*t);
+   earth_z_0+=     0.00000000148 * cos( 0.65447253687 +    95480.94718417450*t);
+   earth_z_0+=     0.00000000121 * cos( 3.13748782783 +    14712.31711645800*t);
+   earth_z_0+=     0.00000000119 * cos( 5.92110458985 +    33019.02111220460*t);
+   earth_z_0+=     0.00000000124 * cos( 2.93638421840 +    32217.20018108080*t);
+   earth_z_0+=     0.00000000117 * cos( 3.65065271640 +    45585.17281218740*t);
+   earth_z_0+=     0.00000000127 * cos( 4.74596574209 +    49515.38250840700*t);
+   earth_z_0+=     0.00000000120 * cos( 1.04211499785 +     6915.85958930460*t);
+   earth_z_0+=     0.00000000120 * cos( 5.60638811846 +     5650.29211067820*t);
+   earth_z_0+=     0.00000000127 * cos( 2.70766997307 +    12352.85260454480*t);
+   earth_z_0+=     0.00000000113 * cos( 5.39691616452 +     5235.32853823670*t);
+   earth_z_0+=     0.00000000097 * cos( 1.07908724794 +     9917.69687450980*t);
+   earth_z_0+=     0.00000000123 * cos( 5.15992792934 +    27511.46787353720*t);
+   earth_z_0+=     0.00000000099 * cos( 4.45774681732 +     4933.20844033260*t);
+   earth_z_0+=     0.00000000087 * cos( 3.93637812950 +    17654.78053974960*t);
+   earth_z_0+=     0.00000000122 * cos( 2.23956068680 +    83997.09113559539*t);
+   earth_z_0+=     0.00000000102 * cos( 3.97386522171 +    10818.13528691580*t);
+   earth_z_0+=     0.00000000084 * cos( 4.21241958885 +    22779.43724619380*t);
+   earth_z_0+=     0.00000000103 * cos( 4.63519841733 +     1349.86740965880*t);
+   earth_z_0+=     0.00000000081 * cos( 4.01655240430 +     1592.59601363280*t);
+   earth_z_0+=     0.00000000092 * cos( 3.24182657569 +    78051.58573131690*t);
+   earth_z_0+=     0.00000000094 * cos( 4.33873972800 +    36147.40987730040*t);
+   earth_z_0+=     0.00000000076 * cos( 2.83897554420 +    17260.15465469040*t);
+   earth_z_0+=     0.00000000100 * cos( 6.07733097102 +    26735.94526221320*t);
+   earth_z_0+=     0.00000000078 * cos( 0.72462427283 +    12779.45079542080*t);
+   earth_z_0+=     0.00000000086 * cos( 5.26029638250 +    28313.28880466100*t);
+   earth_z_0+=     0.00000000093 * cos( 4.31900620254 +    44809.65020086340*t);
+   earth_z_0+=     0.00000000076 * cos( 6.22743405935 +    13521.75144159140*t);
+   earth_z_0+=     0.00000000081 * cos( 2.16420552322 +    13095.84266507740*t);
+   earth_z_0+=     0.00000000100 * cos( 1.38002787119 +    28286.99048486120*t);
+   earth_z_0+=     0.00000000072 * cos( 1.55820597747 +     6256.77753019160*t);
+   earth_z_0+=     0.00000000082 * cos( 4.95202664555 +    10575.40668294180*t);
+   earth_z_0+=     0.00000000082 * cos( 1.69647647075 +     1990.74501704100*t);
+   earth_z_0+=     0.00000000096 * cos( 6.16038106485 +    24356.78078864160*t);
+   earth_z_0+=     0.00000000075 * cos( 2.29836095644 +     3634.62102451840*t);
+   earth_z_0+=     0.00000000075 * cos( 2.66367876557 +    16200.77272450120*t);
+   earth_z_0+=     0.00000000087 * cos( 0.26630214764 +    31441.67756975680*t);
+   earth_z_0+=     0.00000000080 * cos( 5.00001604436 +   150192.21439800429*t);
+   earth_z_0+=     0.00000000079 * cos( 0.01398391548 +    90617.73743129970*t);
+   earth_z_0+=     0.00000000054 * cos( 5.59738773448 +   161338.50000087050*t);
+   earth_z_0+=     0.00000000056 * cos( 2.60133794851 +    73188.37597844210*t);
+   earth_z_0+=     0.00000000055 * cos( 5.81483150022 +   143233.51002162008*t);
+   earth_z_0+=     0.00000000054 * cos( 3.38482031504 +   323049.11878710288*t);
+
+   double earth_z_1=0.0;
+   earth_z_1+=     0.00227822442 * cos( 3.41372504278 +     6283.07584999140*t);
+   earth_z_1+=     0.00005429282 * cos( 0.00000000000 +        0.00000000000*t);
+   earth_z_1+=     0.00001903183 * cos( 3.37061270964 +    12566.15169998280*t);
+   earth_z_1+=     0.00000023859 * cos( 3.32836261978 +    18849.22754997420*t);
+   earth_z_1+=     0.00000008671 * cos( 1.82289925035 +     5223.69391980220*t);
+   earth_z_1+=     0.00000007767 * cos( 3.88787435832 +     5507.55323866740*t);
+   earth_z_1+=     0.00000007063 * cos( 5.19303169379 +     2352.86615377180*t);
+   earth_z_1+=     0.00000007090 * cos( 2.32907317990 +    83996.84731811189*t);
+   earth_z_1+=     0.00000005362 * cos( 2.22429221689 +      529.69096509460*t);
+   earth_z_1+=     0.00000003894 * cos( 2.15568517178 +     6279.55273164240*t);
+   earth_z_1+=     0.00000003892 * cos( 1.53021064904 +     6286.59896834040*t);
+   earth_z_1+=     0.00000003285 * cos( 4.88047587852 +    10213.28554621100*t);
+   earth_z_1+=     0.00000002165 * cos( 6.23220223547 +     7860.41939243920*t);
+   earth_z_1+=     0.00000002084 * cos( 5.84682067018 +     1577.34354244780*t);
+   earth_z_1+=     0.00000001959 * cos( 1.42270916062 +    12036.46073488820*t);
+   earth_z_1+=     0.00000001396 * cos( 0.50057569258 +     6309.37416979120*t);
+   earth_z_1+=     0.00000001584 * cos( 1.81185253086 +     5856.47765911540*t);
+   earth_z_1+=     0.00000001493 * cos( 1.75695182457 +     5884.92684658320*t);
+   earth_z_1+=     0.00000001713 * cos( 6.04661123493 +      398.14900340820*t);
+   earth_z_1+=     0.00000001298 * cos( 3.71598693704 +     1059.38193018920*t);
+   earth_z_1+=     0.00000001573 * cos( 2.20074999518 +     9437.76293488700*t);
+   earth_z_1+=     0.00000001503 * cos( 3.94648902715 +     4705.73230754360*t);
+   earth_z_1+=     0.00000001041 * cos( 2.94815263534 +     6256.77753019160*t);
+   earth_z_1+=     0.00000000983 * cos( 2.24865863102 +     6681.22485339960*t);
+   earth_z_1+=     0.00000001150 * cos( 2.43940396435 +      775.52261132400*t);
+   earth_z_1+=     0.00000001171 * cos( 0.74619691355 +     6812.76681508600*t);
+   earth_z_1+=     0.00000000881 * cos( 1.86103443124 +     3930.20969621960*t);
+   earth_z_1+=     0.00000000696 * cos( 1.45896954690 +    11790.62908865880*t);
+   earth_z_1+=     0.00000000732 * cos( 1.60348898963 +     4694.00295470760*t);
+   earth_z_1+=     0.00000000487 * cos( 0.04282116718 +    11506.76976979360*t);
+   earth_z_1+=     0.00000000632 * cos( 4.34207925908 +     5753.38488489680*t);
+   earth_z_1+=     0.00000000584 * cos( 3.11517364836 +      709.93304855830*t);
+   earth_z_1+=     0.00000000583 * cos( 3.92560802577 +    12168.00269657460*t);
+   earth_z_1+=     0.00000000409 * cos( 4.63047700185 +      796.29800681640*t);
+   earth_z_1+=     0.00000000483 * cos( 5.74980419879 +    17789.84561978500*t);
+   earth_z_1+=     0.00000000405 * cos( 1.97558286872 +     6283.00853968860*t);
+   earth_z_1+=     0.00000000405 * cos( 1.53147989887 +     6283.14316029419*t);
+   earth_z_1+=     0.00000000354 * cos( 3.28490469609 +    25132.30339996560*t);
+   earth_z_1+=     0.00000000417 * cos( 3.63631064623 +    13367.97263110660*t);
+   earth_z_1+=     0.00000000367 * cos( 1.55081387384 +    18073.70493865020*t);
+   earth_z_1+=     0.00000000284 * cos( 0.15061211477 +     3738.76143010800*t);
+   earth_z_1+=     0.00000000295 * cos( 2.75399172194 +    14143.49524243060*t);
+   earth_z_1+=     0.00000000288 * cos( 2.98400521096 +     6275.96230299060*t);
+   earth_z_1+=     0.00000000353 * cos( 2.42503121206 +      213.29909543800*t);
+   earth_z_1+=     0.00000000250 * cos( 0.42640065085 +     7079.37385680780*t);
+   earth_z_1+=     0.00000000251 * cos( 2.81315684448 +     6290.18939699220*t);
+   earth_z_1+=     0.00000000338 * cos( 2.55274440586 +     7058.59846131540*t);
+   earth_z_1+=     0.00000000230 * cos( 2.54182122331 +     7084.89678111520*t);
+   earth_z_1+=     0.00000000212 * cos( 4.71778046808 +     6496.37494542940*t);
+   earth_z_1+=     0.00000000188 * cos( 3.17403855561 +     5486.77784317500*t);
+   earth_z_1+=     0.00000000184 * cos( 4.29245958347 +     1589.07289528380*t);
+   earth_z_1+=     0.00000000216 * cos( 2.95054003179 +    22003.91463486980*t);
+   earth_z_1+=     0.00000000165 * cos( 1.19407165511 +     5643.17856367740*t);
+   earth_z_1+=     0.00000000196 * cos( 4.68462471621 +   167283.76158766549*t);
+   earth_z_1+=     0.00000000162 * cos( 2.94100342992 +    15720.83878487840*t);
+   earth_z_1+=     0.00000000189 * cos( 5.09534739493 +    17298.18232732620*t);
+   earth_z_1+=     0.00000000149 * cos( 5.23057852418 +    11769.85369316640*t);
+   earth_z_1+=     0.00000000150 * cos( 1.65291444016 +     3128.38876509580*t);
+   earth_z_1+=     0.00000000143 * cos( 5.48362795500 +    10977.07880469900*t);
+   earth_z_1+=     0.00000000144 * cos( 3.81203756929 +     6262.30045449900*t);
+   earth_z_1+=     0.00000000143 * cos( 5.97809021355 +     6303.85124548380*t);
+   earth_z_1+=     0.00000000136 * cos( 4.97038154072 +     4136.91043351620*t);
+   earth_z_1+=     0.00000000163 * cos( 2.00627600582 +     5230.80746680300*t);
+   earth_z_1+=     0.00000000135 * cos( 3.68020950012 +        7.11354700080*t);
+   earth_z_1+=     0.00000000130 * cos( 4.41428061934 +     6127.65545055720*t);
+   earth_z_1+=     0.00000000123 * cos( 0.93547502051 +     3340.61242669980*t);
+   earth_z_1+=     0.00000000122 * cos( 3.13394683057 +     1194.44701022460*t);
+   earth_z_1+=     0.00000000131 * cos( 4.24961399629 +     6282.09552892320*t);
+   earth_z_1+=     0.00000000131 * cos( 5.54051222995 +     6284.05617105960*t);
+   earth_z_1+=     0.00000000161 * cos( 3.32421999691 +     6283.31966747490*t);
+   earth_z_1+=     0.00000000128 * cos( 4.35489873365 +    25934.12433108940*t);
+   earth_z_1+=     0.00000000113 * cos( 4.42578862670 +    19651.04848109800*t);
+   earth_z_1+=     0.00000000107 * cos( 2.97447343520 +     8827.39026987480*t);
+   earth_z_1+=     0.00000000107 * cos( 0.41198538606 +     5481.25491886760*t);
+   earth_z_1+=     0.00000000131 * cos( 2.29167404623 +     3154.68708489560*t);
+   earth_z_1+=     0.00000000139 * cos( 0.84923405785 +     6069.77675455340*t);
+   earth_z_1+=     0.00000000109 * cos( 3.21932289821 +      801.82093112380*t);
+   earth_z_1+=     0.00000000115 * cos( 3.71248897285 +    11856.21865142450*t);
+   earth_z_1+=     0.00000000086 * cos( 2.78714560645 +     5216.58037280140*t);
+   earth_z_1+=     0.00000000116 * cos( 1.07208569199 +    12352.85260454480*t);
+   earth_z_1+=     0.00000000085 * cos( 4.74199868968 +     8429.24126646660*t);
+   earth_z_1+=     0.00000000112 * cos( 5.29899379351 +     2544.31441988340*t);
+   earth_z_1+=     0.00000000096 * cos( 5.20691825399 +    17260.15465469040*t);
+   earth_z_1+=     0.00000000074 * cos( 4.51170584596 +     8635.94200376320*t);
+   earth_z_1+=     0.00000000097 * cos( 5.50959692365 +    29864.33402730900*t);
+   earth_z_1+=     0.00000000094 * cos( 0.25007575444 +    21228.39202354580*t);
+   earth_z_1+=     0.00000000082 * cos( 3.01994555641 +    20426.57109242200*t);
+   earth_z_1+=     0.00000000071 * cos( 0.47420303374 +     2146.16541647520*t);
+   earth_z_1+=     0.00000000074 * cos( 1.27647360807 +     4164.31198961300*t);
+   earth_z_1+=     0.00000000075 * cos( 0.73893922793 +     9225.53927328300*t);
+   earth_z_1+=     0.00000000065 * cos( 1.39694123605 +    12569.67481833180*t);
+   earth_z_1+=     0.00000000074 * cos( 2.20184828895 +    37724.75341974820*t);
+   earth_z_1+=     0.00000000058 * cos( 5.36736393336 +    11015.10647733480*t);
+   earth_z_1+=     0.00000000069 * cos( 3.56483879750 +     4535.05943692440*t);
+   earth_z_1+=     0.00000000064 * cos( 2.05206957553 +    12562.62858163380*t);
+   earth_z_1+=     0.00000000050 * cos( 4.30582312752 +      426.59819087600*t);
+   earth_z_1+=     0.00000000056 * cos( 4.14424591372 +     5088.62883976680*t);
+   earth_z_1+=     0.00000000061 * cos( 0.63918772258 +    33794.54372352860*t);
+   earth_z_1+=     0.00000000049 * cos( 6.20603791837 +    12559.03815298200*t);
+   earth_z_1+=     0.00000000059 * cos( 4.92457810620 +    71430.69561812909*t);
+   earth_z_1+=     0.00000000058 * cos( 1.68201284430 +    25158.60171976540*t);
+   earth_z_1+=     0.00000000051 * cos( 5.03600618727 +     5863.59120611620*t);
+   earth_z_1+=     0.00000000047 * cos( 0.30276119177 +    11371.70468975820*t);
+   earth_z_1+=     0.00000000050 * cos( 0.77635440262 +       26.29831979980*t);
+   earth_z_1+=     0.00000000041 * cos( 1.87344929139 +     1748.01641306700*t);
+   earth_z_1+=     0.00000000041 * cos( 5.91861144924 +    23581.25817731760*t);
+   earth_z_1+=     0.00000000055 * cos( 1.37242658028 +     7342.45778018060*t);
+   earth_z_1+=     0.00000000053 * cos( 1.47750772846 +    18319.53658487960*t);
+   earth_z_1+=     0.00000000042 * cos( 0.11693102263 +    23543.23050468179*t);
+   earth_z_1+=     0.00000000045 * cos( 5.18332998014 +    16496.36139620240*t);
+   earth_z_1+=     0.00000000045 * cos( 1.69577010458 +     6702.56049386660*t);
+   earth_z_1+=     0.00000000050 * cos( 3.62740835096 +    41654.96311596780*t);
+   earth_z_1+=     0.00000000050 * cos( 0.12043698146 +    29088.81141598500*t);
+   earth_z_1+=     0.00000000039 * cos( 2.15562950918 +     4933.20844033260*t);
+   earth_z_1+=     0.00000000048 * cos( 5.37118742490 +     7477.52286021600*t);
+   earth_z_1+=     0.00000000040 * cos( 3.02033323932 +    10447.38783960440*t);
+   earth_z_1+=     0.00000000035 * cos( 3.52010534689 +    26087.90314157420*t);
+   earth_z_1+=     0.00000000045 * cos( 6.10982446357 +     8031.09226305840*t);
+   earth_z_1+=     0.00000000043 * cos( 4.08726331262 +   156137.47598479928*t);
+   earth_z_1+=     0.00000000032 * cos( 6.09025731476 +    64809.80550494129*t);
+   earth_z_1=earth_z_1 * t;
+
+   double earth_z_2=0.0;
+   earth_z_2+=     0.00009721989 * cos( 5.15233725915 +     6283.07584999140*t);
+   earth_z_2+=     0.00000349501 * cos( 3.14159265359 +        0.00000000000*t);
+   earth_z_2+=     0.00000067136 * cos( 0.64403888586 +    12566.15169998280*t);
+   earth_z_2+=     0.00000002169 * cos( 1.07314429736 +    18849.22754997420*t);
+   earth_z_2+=     0.00000001662 * cos( 1.62746869551 +    84334.66158130829*t);
+   earth_z_2+=     0.00000000635 * cos( 3.51985338656 +     6279.55273164240*t);
+   earth_z_2+=     0.00000000492 * cos( 2.41382223971 +     1047.74731175470*t);
+   earth_z_2+=     0.00000000347 * cos( 6.16127984865 +     5223.69391980220*t);
+   earth_z_2+=     0.00000000322 * cos( 0.37660897389 +     6286.59896834040*t);
+   earth_z_2+=     0.00000000302 * cos( 4.06770618987 +    83996.84731811189*t);
+   earth_z_2+=     0.00000000314 * cos( 2.37089871705 +     5507.55323866740*t);
+   earth_z_2+=     0.00000000242 * cos( 1.64810071161 +     7860.41939243920*t);
+   earth_z_2+=     0.00000000243 * cos( 3.96741961391 +      529.69096509460*t);
+   earth_z_2+=     0.00000000229 * cos( 0.50869157074 +     2352.86615377180*t);
+   earth_z_2+=     0.00000000170 * cos( 5.28668290523 +     6256.77753019160*t);
+   earth_z_2+=     0.00000000174 * cos( 0.13496987900 +    10213.28554621100*t);
+   earth_z_2+=     0.00000000125 * cos( 3.13799899281 +     5884.92684658320*t);
+   earth_z_2+=     0.00000000118 * cos( 3.85466055344 +     1577.34354244780*t);
+   earth_z_2+=     0.00000000123 * cos( 4.17102530625 +     6309.37416979120*t);
+   earth_z_2+=     0.00000000089 * cos( 1.65074757807 +     4705.73230754360*t);
+   earth_z_2+=     0.00000000090 * cos( 4.13333302298 +     5753.38488489680*t);
+   earth_z_2+=     0.00000000067 * cos( 0.30879908581 +      775.52261132400*t);
+   earth_z_2+=     0.00000000088 * cos( 2.93230657226 +    12036.46073488820*t);
+   earth_z_2+=     0.00000000065 * cos( 6.02619905507 +     6812.76681508600*t);
+   earth_z_2+=     0.00000000067 * cos( 1.15131045651 +    25132.30339996560*t);
+   earth_z_2+=     0.00000000065 * cos( 2.72055837855 +     6127.65545055720*t);
+   earth_z_2+=     0.00000000062 * cos( 0.78482367183 +     6438.49624942560*t);
+   earth_z_2+=     0.00000000057 * cos( 1.49017558261 +    11506.76976979360*t);
+   earth_z_2+=     0.00000000056 * cos( 4.55921214755 +     5486.77784317500*t);
+   earth_z_2+=     0.00000000075 * cos( 0.84213523741 +   167621.57585086189*t);
+   earth_z_2+=     0.00000000069 * cos( 1.21340931171 +      398.14900340820*t);
+   earth_z_2+=     0.00000000047 * cos( 5.28931662006 +     7079.37385680780*t);
+   earth_z_2+=     0.00000000044 * cos( 3.09697268535 +    11790.62908865880*t);
+   earth_z_2+=     0.00000000057 * cos( 6.15295833679 +    12194.03291462090*t);
+   earth_z_2+=     0.00000000040 * cos( 1.74073861374 +     1059.38193018920*t);
+   earth_z_2+=     0.00000000040 * cos( 2.70069983169 +    10988.80815753500*t);
+   earth_z_2+=     0.00000000036 * cos( 0.32483893578 +     6681.22485339960*t);
+   earth_z_2+=     0.00000000046 * cos( 3.38617099014 +   156475.29024799570*t);
+   earth_z_2+=     0.00000000033 * cos( 1.26266496002 +     6290.18939699220*t);
+   earth_z_2+=     0.00000000033 * cos( 2.25067065498 +     6275.96230299060*t);
+   earth_z_2+=     0.00000000032 * cos( 2.90595681986 +      796.29800681640*t);
+   earth_z_2+=     0.00000000029 * cos( 1.96179670465 +        7.11354700080*t);
+   earth_z_2+=     0.00000000023 * cos( 4.86623778952 +     3738.76143010800*t);
+   earth_z_2+=     0.00000000025 * cos( 4.85403147726 +      709.93304855830*t);
+   earth_z_2+=     0.00000000021 * cos( 6.22275008403 +     6069.77675455340*t);
+   earth_z_2+=     0.00000000024 * cos( 5.83690442827 +     6282.09552892320*t);
+   earth_z_2+=     0.00000000024 * cos( 3.95322179797 +     6284.05617105960*t);
+   earth_z_2+=     0.00000000022 * cos( 3.46618662246 +     4136.91043351620*t);
+   earth_z_2+=     0.00000000021 * cos( 3.58048145762 +     6496.37494542940*t);
+   earth_z_2+=     0.00000000025 * cos( 5.59503085104 +     7058.59846131540*t);
+   earth_z_2+=     0.00000000019 * cos( 1.62978232282 +     1194.44701022460*t);
+   earth_z_2+=     0.00000000019 * cos( 2.82679050999 +     1589.07289528380*t);
+   earth_z_2+=     0.00000000024 * cos( 5.40715542605 +    12168.00269657460*t);
+   earth_z_2=earth_z_2 * t * t;
+
+   double earth_z_3=0.0;
+   earth_z_3+=     0.00000276077 * cos( 0.59413258730 +     6283.07584999140*t);
+   earth_z_3+=     0.00000025551 * cos( 3.14159265359 +        0.00000000000*t);
+   earth_z_3+=     0.00000001810 * cos( 0.11612262117 +    12566.15169998280*t);
+   earth_z_3+=     0.00000000113 * cos( 5.66152847225 +    18849.22754997420*t);
+   earth_z_3+=     0.00000000056 * cos( 5.02765554835 +     6279.55273164240*t);
+   earth_z_3+=     0.00000000024 * cos( 0.41571093398 +    83996.84731811189*t);
+   earth_z_3+=     0.00000000019 * cos( 5.99007646261 +     6256.77753019160*t);
+   earth_z_3+=     0.00000000018 * cos( 3.80004734567 +     6309.37416979120*t);
+   earth_z_3+=     0.00000000018 * cos( 1.21049250774 +     6127.65545055720*t);
+   earth_z_3+=     0.00000000018 * cos( 2.29734567137 +     6438.49624942560*t);
+   earth_z_3+=     0.00000000015 * cos( 4.72881467263 +     6286.59896834040*t);
+   earth_z_3+=     0.00000000011 * cos( 1.48347223110 +      709.93304855830*t);
+   earth_z_3=earth_z_3 * t * t * t;
+
+   double earth_z_4=0.0;
+   earth_z_4+=     0.00000005751 * cos( 2.27069090892 +     6283.07584999140*t);
+   earth_z_4+=     0.00000001305 * cos( 0.00000000000 +        0.00000000000*t);
+   earth_z_4+=     0.00000000060 * cos( 4.26712050846 +    12566.15169998280*t);
+   earth_z_4+=     0.00000000006 * cos( 4.13914011953 +    18849.22754997420*t);
+   earth_z_4+=     0.00000000005 * cos( 0.84308705203 +     1047.74731175470*t);
+   earth_z_4+=     0.00000000005 * cos( 0.05711572303 +    84334.66158130829*t);
+   earth_z_4=earth_z_4 * t * t * t * t;
+
+   double earth_z_5=0.0;
+   earth_z_5+=     0.00000000114 * cos( 4.31169156923 +     6283.07584999140*t);
+   earth_z_5+=     0.00000000035 * cos( 0.00000000000 +        0.00000000000*t);
+   earth_z_5=earth_z_5 * t * t * t * t * t;
+
+   return earth_z_0+earth_z_1+earth_z_2+earth_z_3+earth_z_4+earth_z_5;
+}
+
+double vsop87a_full::emb_x(double t){
+   double emb_x_0=0.0;
+   emb_x_0+=     0.99982927460 * cos( 1.75348568475 +     6283.07584999140*t);
+   emb_x_0+=     0.00835257300 * cos( 1.71034539450 +    12566.15169998280*t);
+   emb_x_0+=     0.00561144161 * cos( 0.00000000000 +        0.00000000000*t);
+   emb_x_0+=     0.00010466628 * cos( 1.66722645223 +    18849.22754997420*t);
+   emb_x_0+=     0.00002552498 * cos( 0.58310207301 +      529.69096509460*t);
+   emb_x_0+=     0.00002137256 * cos( 1.09235189672 +     1577.34354244780*t);
+   emb_x_0+=     0.00001709103 * cos( 0.49540223397 +     6279.55273164240*t);
+   emb_x_0+=     0.00001707882 * cos( 6.15315547484 +     6286.59896834040*t);
+   emb_x_0+=     0.00001445242 * cos( 3.47272783760 +     2352.86615377180*t);
+   emb_x_0+=     0.00001091006 * cos( 3.68984782465 +     5223.69391980220*t);
+   emb_x_0+=     0.00000934429 * cos( 6.07389922585 +    12036.46073488820*t);
+   emb_x_0+=     0.00000899144 * cos( 3.17571950523 +    10213.28554621100*t);
+   emb_x_0+=     0.00000566514 * cos( 2.15262034016 +     1059.38193018920*t);
+   emb_x_0+=     0.00000684416 * cos( 1.30699021227 +     5753.38488489680*t);
+   emb_x_0+=     0.00000734455 * cos( 4.35500196530 +      398.14900340820*t);
+   emb_x_0+=     0.00000681437 * cos( 2.21821534685 +     4705.73230754360*t);
+   emb_x_0+=     0.00000611238 * cos( 5.38479234323 +     6812.76681508600*t);
+   emb_x_0+=     0.00000451836 * cos( 6.08768280868 +     5884.92684658320*t);
+   emb_x_0+=     0.00000451953 * cos( 1.27933728354 +     6256.77753019160*t);
+   emb_x_0+=     0.00000449517 * cos( 5.36923831714 +     6309.37416979120*t);
+   emb_x_0+=     0.00000406248 * cos( 0.54361367084 +     6681.22485339960*t);
+   emb_x_0+=     0.00000540957 * cos( 0.78677364655 +      775.52261132400*t);
+   emb_x_0+=     0.00000547004 * cos( 1.46146650376 +    14143.49524243060*t);
+   emb_x_0+=     0.00000520484 * cos( 4.43295799975 +     7860.41939243920*t);
+   emb_x_0+=     0.00000214960 * cos( 4.50213844573 +    11506.76976979360*t);
+   emb_x_0+=     0.00000227892 * cos( 1.23941482802 +     7058.59846131540*t);
+   emb_x_0+=     0.00000225878 * cos( 3.27244306207 +     4694.00295470760*t);
+   emb_x_0+=     0.00000255820 * cos( 2.26556277246 +    12168.00269657460*t);
+   emb_x_0+=     0.00000178120 * cos( 2.96205424204 +      796.29800681640*t);
+   emb_x_0+=     0.00000161205 * cos( 1.47337718956 +     5486.77784317500*t);
+   emb_x_0+=     0.00000178325 * cos( 6.24374704602 +     6283.14316029419*t);
+   emb_x_0+=     0.00000178325 * cos( 0.40466470869 +     6283.00853968860*t);
+   emb_x_0+=     0.00000155487 * cos( 1.62409309523 +    25132.30339996560*t);
+   emb_x_0+=     0.00000209024 * cos( 5.85207528073 +    11790.62908865880*t);
+   emb_x_0+=     0.00000199971 * cos( 4.07209938245 +    17789.84561978500*t);
+   emb_x_0+=     0.00000128933 * cos( 5.21693314150 +     7079.37385680780*t);
+   emb_x_0+=     0.00000128099 * cos( 4.80182882228 +     3738.76143010800*t);
+   emb_x_0+=     0.00000151691 * cos( 0.86921639327 +      213.29909543800*t);
+   emb_x_0+=     0.00000094483 * cos( 4.60187620729 +    10977.07880469900*t);
+   emb_x_0+=     0.00000077801 * cos( 1.84444797200 +     8827.39026987480*t);
+   emb_x_0+=     0.00000077277 * cos( 3.58396027325 +     5507.55323866740*t);
+   emb_x_0+=     0.00000073493 * cos( 2.69522392584 +     1589.07289528380*t);
+   emb_x_0+=     0.00000065359 * cos( 3.65156347970 +    11769.85369316640*t);
+   emb_x_0+=     0.00000063240 * cos( 2.24124130944 +     6262.30045449900*t);
+   emb_x_0+=     0.00000062961 * cos( 4.40729381941 +     6303.85124548380*t);
+   emb_x_0+=     0.00000082957 * cos( 6.19254908840 +     3340.61242669980*t);
+   emb_x_0+=     0.00000079201 * cos( 2.49366716519 +     6069.77675455340*t);
+   emb_x_0+=     0.00000060048 * cos( 3.39951805979 +     4136.91043351620*t);
+   emb_x_0+=     0.00000062119 * cos( 2.01472298915 +        7.11354700080*t);
+   emb_x_0+=     0.00000075601 * cos( 4.15956023781 +     6496.37494542940*t);
+   emb_x_0+=     0.00000067771 * cos( 4.03416283669 +     9437.76293488700*t);
+   emb_x_0+=     0.00000053715 * cos( 1.56222706033 +     1194.44701022460*t);
+   emb_x_0+=     0.00000057574 * cos( 2.67881766950 +     6282.09552892320*t);
+   emb_x_0+=     0.00000057574 * cos( 3.96971590315 +     6284.05617105960*t);
+   emb_x_0+=     0.00000051129 * cos( 1.31255522350 +     6290.18939699220*t);
+   emb_x_0+=     0.00000050529 * cos( 5.34827650648 +     6275.96230299060*t);
+   emb_x_0+=     0.00000049613 * cos( 2.64370071484 +     6127.65545055720*t);
+   emb_x_0+=     0.00000048085 * cos( 4.00886197380 +     6438.49624942560*t);
+   emb_x_0+=     0.00000057340 * cos( 3.79459979103 +     3154.68708489560*t);
+   emb_x_0+=     0.00000047378 * cos( 6.10705175562 +     3128.38876509580*t);
+   emb_x_0+=     0.00000048104 * cos( 4.77180566934 +      801.82093112380*t);
+   emb_x_0+=     0.00000041149 * cos( 3.32726302190 +     8429.24126646660*t);
+   emb_x_0+=     0.00000052309 * cos( 5.30571599873 +    13367.97263110660*t);
+   emb_x_0+=     0.00000051321 * cos( 5.78450026632 +    12352.85260454480*t);
+   emb_x_0+=     0.00000043381 * cos( 3.64501749522 +    17260.15465469040*t);
+   emb_x_0+=     0.00000039554 * cos( 4.93022392385 +     5481.25491886760*t);
+   emb_x_0+=     0.00000048969 * cos( 0.45418080759 +     9225.53927328300*t);
+   emb_x_0+=     0.00000047568 * cos( 3.16123585999 +     5856.47765911540*t);
+   emb_x_0+=     0.00000047271 * cos( 0.62147119447 +     2544.31441988340*t);
+   emb_x_0+=     0.00000037987 * cos( 3.04100734504 +      426.59819087600*t);
+   emb_x_0+=     0.00000032481 * cos( 0.80957263520 +     3930.20969621960*t);
+   emb_x_0+=     0.00000032559 * cos( 1.97427864771 +     2146.16541647520*t);
+   emb_x_0+=     0.00000032521 * cos( 2.84566871337 +     4164.31198961300*t);
+   emb_x_0+=     0.00000032546 * cos( 3.01796808307 +     5088.62883976680*t);
+   emb_x_0+=     0.00000028515 * cos( 6.10933021271 +    12569.67481833180*t);
+   emb_x_0+=     0.00000036898 * cos( 2.91097134183 +    18073.70493865020*t);
+   emb_x_0+=     0.00000024936 * cos( 2.17935828787 +       26.29831979980*t);
+   emb_x_0+=     0.00000030322 * cos( 1.99340330773 +     4535.05943692440*t);
+   emb_x_0+=     0.00000028166 * cos( 0.48127325250 +    12562.62858163380*t);
+   emb_x_0+=     0.00000021880 * cos( 5.12635088590 +    11371.70468975820*t);
+   emb_x_0+=     0.00000021642 * cos( 4.63524156478 +    12559.03815298200*t);
+   emb_x_0+=     0.00000020480 * cos( 1.46713660207 +     7084.89678111520*t);
+   emb_x_0+=     0.00000022867 * cos( 4.76437399404 +     8031.09226305840*t);
+   emb_x_0+=     0.00000018593 * cos( 3.38415795860 +     1748.01641306700*t);
+   emb_x_0+=     0.00000023250 * cos( 6.19030311409 +    18319.53658487960*t);
+   emb_x_0+=     0.00000017104 * cos( 0.58684867006 +     4933.20844033260*t);
+   emb_x_0+=     0.00000021687 * cos( 4.30348901447 +    10447.38783960440*t);
+   emb_x_0+=     0.00000021065 * cos( 3.80028088787 +     7477.52286021600*t);
+   emb_x_0+=     0.00000014285 * cos( 1.29528607840 +     2942.46342329160*t);
+   emb_x_0+=     0.00000013898 * cos( 4.59171958613 +     8635.94200376320*t);
+   emb_x_0+=     0.00000019220 * cos( 0.49428010004 +    17298.18232732620*t);
+   emb_x_0+=     0.00000017889 * cos( 0.15750451509 +     1592.59601363280*t);
+   emb_x_0+=     0.00000013605 * cos( 5.19721055627 +    13095.84266507740*t);
+   emb_x_0+=     0.00000015049 * cos( 3.12084096050 +    16496.36139620240*t);
+   emb_x_0+=     0.00000013852 * cos( 6.16243027421 +     7632.94325965020*t);
+   emb_x_0+=     0.00000014421 * cos( 1.43791851415 +    20426.57109242200*t);
+   emb_x_0+=     0.00000013247 * cos( 3.60768174652 +    12139.55350910680*t);
+   emb_x_0+=     0.00000011583 * cos( 5.46076026245 +     5331.35744374080*t);
+   emb_x_0+=     0.00000014329 * cos( 6.02791958133 +     7342.45778018060*t);
+   emb_x_0+=     0.00000012375 * cos( 3.10941923273 +     6279.48542133960*t);
+   emb_x_0+=     0.00000012366 * cos( 3.53911446434 +     6286.66627864320*t);
+   emb_x_0+=     0.00000009951 * cos( 1.25735916596 +     7234.79425624200*t);
+   emb_x_0+=     0.00000012425 * cos( 1.06396563046 +    15110.46611986620*t);
+   emb_x_0+=     0.00000011004 * cos( 2.19276723240 +    10988.80815753500*t);
+   emb_x_0+=     0.00000011583 * cos( 4.05256206393 +     5729.50644714900*t);
+   emb_x_0+=     0.00000009046 * cos( 5.43016645841 +     9623.68827669120*t);
+   emb_x_0+=     0.00000008869 * cos( 5.04690941813 +     6148.01076995600*t);
+   emb_x_0+=     0.00000010507 * cos( 2.62758025060 +     6836.64525283380*t);
+   emb_x_0+=     0.00000011129 * cos( 4.17728063171 +    15720.83878487840*t);
+   emb_x_0+=     0.00000008614 * cos( 1.60161176919 +     6418.14093002680*t);
+   emb_x_0+=     0.00000008532 * cos( 2.46198384967 +    14712.31711645800*t);
+   emb_x_0+=     0.00000007910 * cos( 3.15429574843 +     2118.76386037840*t);
+   emb_x_0+=     0.00000010546 * cos( 4.79503596728 +     1349.86740965880*t);
+   emb_x_0+=     0.00000010480 * cos( 2.95288908781 +     5999.21653112620*t);
+   emb_x_0+=     0.00000007456 * cos( 5.42522643084 +     6040.34724601740*t);
+   emb_x_0+=     0.00000009820 * cos( 3.69556194245 +     6566.93516885660*t);
+   emb_x_0+=     0.00000009171 * cos( 3.92216439653 +     5643.17856367740*t);
+   emb_x_0+=     0.00000007018 * cos( 1.23420199826 +     6525.80445396540*t);
+   emb_x_0+=     0.00000009797 * cos( 1.91951822764 +    21228.39202354580*t);
+   emb_x_0+=     0.00000009060 * cos( 4.61643194755 +     4690.47983635860*t);
+   emb_x_0+=     0.00000007545 * cos( 1.23604549224 +    12539.85338018300*t);
+   emb_x_0+=     0.00000008227 * cos( 0.32929613711 +    10973.55568635000*t);
+   emb_x_0+=     0.00000007202 * cos( 5.88391543785 +     6245.04817735560*t);
+   emb_x_0+=     0.00000007144 * cos( 0.76464650349 +     6321.10352262720*t);
+   emb_x_0+=     0.00000006625 * cos( 0.70112927898 +     5327.47610838280*t);
+   emb_x_0+=     0.00000007412 * cos( 5.03197290869 +    23543.23050468179*t);
+   emb_x_0+=     0.00000006160 * cos( 0.94570829909 +    12964.30070339100*t);
+   emb_x_0+=     0.00000006378 * cos( 6.20903842618 +      951.71840625060*t);
+   emb_x_0+=     0.00000006359 * cos( 5.03541261021 +     1990.74501704100*t);
+   emb_x_0+=     0.00000007254 * cos( 3.57690549862 +      536.80451209540*t);
+   emb_x_0+=     0.00000005777 * cos( 6.10164059960 +      955.59974160860*t);
+   emb_x_0+=     0.00000006393 * cos( 3.97543391449 +    24072.92146977640*t);
+   emb_x_0+=     0.00000006927 * cos( 0.71517687323 +      522.57741809380*t);
+   emb_x_0+=     0.00000007793 * cos( 4.30411864972 +    22003.91463486980*t);
+   emb_x_0+=     0.00000006346 * cos( 3.50687769075 +     5230.80746680300*t);
+   emb_x_0+=     0.00000007534 * cos( 3.51457135016 +    18422.62935909819*t);
+   emb_x_0+=     0.00000006953 * cos( 2.74384548965 +      155.42039943420*t);
+   emb_x_0+=     0.00000006224 * cos( 2.24416723229 +    18451.07854656599*t);
+   emb_x_0+=     0.00000005830 * cos( 2.32860514691 +      639.89728631400*t);
+   emb_x_0+=     0.00000005762 * cos( 3.49866583471 +     1551.04522264800*t);
+   emb_x_0+=     0.00000006251 * cos( 4.45479090841 +     6288.59877429880*t);
+   emb_x_0+=     0.00000006251 * cos( 2.19374266424 +     6277.55292568400*t);
+   emb_x_0+=     0.00000006278 * cos( 0.64265846614 +     5216.58037280140*t);
+   emb_x_0+=     0.00000006405 * cos( 2.25334669847 +     5650.29211067820*t);
+   emb_x_0+=     0.00000006545 * cos( 1.98229464447 +      103.09277421860*t);
+   emb_x_0+=     0.00000004737 * cos( 3.74210997787 +    14314.16811304980*t);
+   emb_x_0+=     0.00000004646 * cos( 1.33492108974 +      553.56940284240*t);
+   emb_x_0+=     0.00000006263 * cos( 4.39529396163 +    26087.90314157420*t);
+   emb_x_0+=     0.00000005414 * cos( 2.99961957880 +     6179.98307577280*t);
+   emb_x_0+=     0.00000005414 * cos( 3.65424820690 +     6386.16862421000*t);
+   emb_x_0+=     0.00000005711 * cos( 4.32714492993 +     6915.85958930460*t);
+   emb_x_0+=     0.00000004269 * cos( 2.38367994999 +     7875.67186362420*t);
+   emb_x_0+=     0.00000004017 * cos( 3.78580907612 +     6254.62666252360*t);
+   emb_x_0+=     0.00000004067 * cos( 2.42033776129 +     3634.62102451840*t);
+   emb_x_0+=     0.00000003993 * cos( 2.86272124248 +     6311.52503745920*t);
+   emb_x_0+=     0.00000005305 * cos( 3.34359172920 +    25158.60171976540*t);
+   emb_x_0+=     0.00000004804 * cos( 2.82827032932 +     5760.49843189760*t);
+   emb_x_0+=     0.00000004349 * cos( 5.62752376053 +     7238.67559160000*t);
+   emb_x_0+=     0.00000004119 * cos( 3.29663921982 +     6709.67404086740*t);
+   emb_x_0+=     0.00000004024 * cos( 0.55473527406 +    11015.10647733480*t);
+   emb_x_0+=     0.00000004409 * cos( 3.80563977803 +     6805.65326808520*t);
+   emb_x_0+=     0.00000003348 * cos( 3.54433469125 +    12592.45001978260*t);
+   emb_x_0+=     0.00000003172 * cos( 5.19146555016 +     1066.49547719000*t);
+   emb_x_0+=     0.00000003532 * cos( 3.92898013298 +     9917.69687450980*t);
+   emb_x_0+=     0.00000003580 * cos( 4.19680692038 +    11926.25441366880*t);
+   emb_x_0+=     0.00000003336 * cos( 1.79793376416 +    10575.40668294180*t);
+   emb_x_0+=     0.00000003149 * cos( 6.21487368875 +     4292.33083295040*t);
+   emb_x_0+=     0.00000003899 * cos( 5.53698997029 +       74.78159856730*t);
+   emb_x_0+=     0.00000002979 * cos( 6.20067573402 +    12566.21901028560*t);
+   emb_x_0+=     0.00000003981 * cos( 1.19685195313 +    13341.67431130680*t);
+   emb_x_0+=     0.00000003025 * cos( 3.18620101831 +    18052.92954315780*t);
+   emb_x_0+=     0.00000003624 * cos( 5.60981856994 +     6208.29425142410*t);
+   emb_x_0+=     0.00000003581 * cos( 1.03995231220 +     6357.85744855870*t);
+   emb_x_0+=     0.00000003698 * cos( 5.22105764023 +     5966.68398033480*t);
+   emb_x_0+=     0.00000002793 * cos( 4.13369432245 +        3.52311834900*t);
+   emb_x_0+=     0.00000003626 * cos( 1.42632163253 +     6599.46771964800*t);
+   emb_x_0+=     0.00000003625 * cos( 3.00192353938 +     6019.99192661860*t);
+   emb_x_0+=     0.00000002611 * cos( 2.43263787635 +     6702.56049386660*t);
+   emb_x_0+=     0.00000002535 * cos( 1.58310693201 +    31415.37924995700*t);
+   emb_x_0+=     0.00000002593 * cos( 4.91225714775 +    16730.46368959580*t);
+   emb_x_0+=     0.00000003434 * cos( 3.64656881585 +     6546.15977336420*t);
+   emb_x_0+=     0.00000002788 * cos( 3.30835728306 +    12721.57209941700*t);
+   emb_x_0+=     0.00000003383 * cos( 1.74748489437 +    14945.31617355440*t);
+   emb_x_0+=     0.00000002653 * cos( 4.42220751136 +     5863.59120611620*t);
+   emb_x_0+=     0.00000003049 * cos( 0.02765893726 +     4732.03062734340*t);
+   emb_x_0+=     0.00000002378 * cos( 3.62936279948 +     2388.89402044920*t);
+   emb_x_0+=     0.00000002979 * cos( 0.36159339669 +    12566.08438968000*t);
+   emb_x_0+=     0.00000002716 * cos( 3.06935305016 +    12029.34718788740*t);
+   emb_x_0+=     0.00000002216 * cos( 5.26633529197 +    13362.44970679920*t);
+   emb_x_0+=     0.00000003014 * cos( 4.76640877174 +    29088.81141598500*t);
+   emb_x_0+=     0.00000002742 * cos( 5.06666852445 +    13916.01910964160*t);
+   emb_x_0+=     0.00000002265 * cos( 5.95706192557 +     6076.89030155420*t);
+   emb_x_0+=     0.00000002529 * cos( 5.94709040579 +    12043.57428188900*t);
+   emb_x_0+=     0.00000002015 * cos( 2.48432243402 +     6262.72053059260*t);
+   emb_x_0+=     0.00000002218 * cos( 0.69324011482 +     6489.26139842860*t);
+   emb_x_0+=     0.00000002005 * cos( 4.16423889100 +     6303.43116939020*t);
+   emb_x_0+=     0.00000002009 * cos( 4.09759324685 +     4701.11650170840*t);
+   emb_x_0+=     0.00000002105 * cos( 4.63993587835 +     5746.27133789600*t);
+   emb_x_0+=     0.00000001937 * cos( 0.22112768064 +    13521.75144159140*t);
+   emb_x_0+=     0.00000002437 * cos( 6.12896014686 +      149.56319713460*t);
+   emb_x_0+=     0.00000002602 * cos( 4.28001549113 +     4590.91018048900*t);
+   emb_x_0+=     0.00000001893 * cos( 5.66795892365 +     6286.95718534940*t);
+   emb_x_0+=     0.00000001893 * cos( 0.98057464900 +     6279.19451463340*t);
+   emb_x_0+=     0.00000002388 * cos( 0.13031678718 +     1162.47470440780*t);
+   emb_x_0+=     0.00000001962 * cos( 0.25060934155 +     4686.88940770680*t);
+   emb_x_0+=     0.00000002337 * cos( 4.40745417534 +    10021.83728009940*t);
+   emb_x_0+=     0.00000001884 * cos( 5.27011523204 +      242.72860397400*t);
+   emb_x_0+=     0.00000001759 * cos( 0.19888470443 +    13517.87010623340*t);
+   emb_x_0+=     0.00000002163 * cos( 1.41797947214 +     6037.24420376200*t);
+   emb_x_0+=     0.00000002167 * cos( 4.72612563566 +    23013.53953958720*t);
+   emb_x_0+=     0.00000002250 * cos( 2.39728930597 +    12669.24447420140*t);
+   emb_x_0+=     0.00000002052 * cos( 5.23006568302 +     6528.90749622080*t);
+   emb_x_0+=     0.00000001774 * cos( 2.08484497348 +     6819.88036208680*t);
+   emb_x_0+=     0.00000001549 * cos( 0.71612589777 +       20.77539549240*t);
+   emb_x_0+=     0.00000001543 * cos( 5.96098287839 +     6272.43918464160*t);
+   emb_x_0+=     0.00000001543 * cos( 0.68755069426 +     6293.71251534120*t);
+   emb_x_0+=     0.00000001727 * cos( 4.50674659690 +    17654.78053974960*t);
+   emb_x_0+=     0.00000001436 * cos( 2.99734831658 +      110.20632121940*t);
+   emb_x_0+=     0.00000001427 * cos( 1.68851809181 +    13119.72110282519*t);
+   emb_x_0+=     0.00000001395 * cos( 3.25347392533 +    10177.25767953360*t);
+   emb_x_0+=     0.00000001446 * cos( 2.26563038586 +     1052.26838318840*t);
+   emb_x_0+=     0.00000001819 * cos( 3.18848907826 +    18875.52586977400*t);
+   emb_x_0+=     0.00000001740 * cos( 5.48710781556 +    19651.04848109800*t);
+   emb_x_0+=     0.00000001794 * cos( 5.74176427744 +    25934.12433108940*t);
+   emb_x_0+=     0.00000001800 * cos( 6.18908014275 +    33019.02111220460*t);
+   emb_x_0+=     0.00000001488 * cos( 5.67443981584 +       76.26607127560*t);
+   emb_x_0+=     0.00000001349 * cos( 5.75097793699 +     6016.46880826960*t);
+   emb_x_0+=     0.00000001477 * cos( 4.14004787537 +    12779.45079542080*t);
+   emb_x_0+=     0.00000001251 * cos( 1.50045769848 +     3894.18182954220*t);
+   emb_x_0+=     0.00000001492 * cos( 6.22975090570 +     9411.46461508720*t);
+   emb_x_0+=     0.00000001471 * cos( 5.75783391755 +    18635.92845453620*t);
+   emb_x_0+=     0.00000001672 * cos( 5.43726223249 +    16200.77272450120*t);
+   emb_x_0+=     0.00000001263 * cos( 0.91728971350 +    12573.26524698360*t);
+   emb_x_0+=     0.00000001366 * cos( 1.09160445123 +     2379.16447357160*t);
+   emb_x_0+=     0.00000001272 * cos( 0.89738866811 +     6549.68289171320*t);
+   emb_x_0+=     0.00000001274 * cos( 5.34253525843 +      632.78373931320*t);
+   emb_x_0+=     0.00000001156 * cos( 1.38847609016 +    11499.65622279280*t);
+   emb_x_0+=     0.00000001543 * cos( 4.26237609038 +    18209.33026366019*t);
+   emb_x_0+=     0.00000001460 * cos( 4.97233669659 +      245.83164622940*t);
+   emb_x_0+=     0.00000001455 * cos( 5.29534379992 +       38.13303563780*t);
+   emb_x_0+=     0.00000001228 * cos( 2.83750345253 +    24356.78078864160*t);
+   emb_x_0+=     0.00000001028 * cos( 2.51030023881 +     6112.40297937220*t);
+   emb_x_0+=     0.00000001268 * cos( 4.67217109478 +     5429.87946823940*t);
+   emb_x_0+=     0.00000001312 * cos( 1.10091569645 +    12491.37010141550*t);
+   emb_x_0+=     0.00000000992 * cos( 4.13817717100 +     6453.74872061060*t);
+   emb_x_0+=     0.00000000988 * cos( 0.98006619973 +     8273.82086703240*t);
+   emb_x_0+=     0.00000001191 * cos( 3.43008039922 +     6290.12208668940*t);
+   emb_x_0+=     0.00000001191 * cos( 3.21845317343 +     6276.02961329339*t);
+   emb_x_0+=     0.00000000978 * cos( 4.73852763851 +     6247.04798331399*t);
+   emb_x_0+=     0.00000000971 * cos( 1.91000718426 +     6319.10371666880*t);
+   emb_x_0+=     0.00000000962 * cos( 2.63568544853 +    12565.17137891460*t);
+   emb_x_0+=     0.00000001025 * cos( 2.21768401152 +    12545.37630449040*t);
+   emb_x_0+=     0.00000001146 * cos( 3.85865009647 +      206.18554843720*t);
+   emb_x_0+=     0.00000000944 * cos( 0.78473810598 +     6306.95428773920*t);
+   emb_x_0+=     0.00000000944 * cos( 5.86379546667 +     6259.19741224360*t);
+   emb_x_0+=     0.00000001092 * cos( 3.92002105501 +    17796.95916678580*t);
+   emb_x_0+=     0.00000000886 * cos( 5.23132430543 +     6172.86952877200*t);
+   emb_x_0+=     0.00000000859 * cos( 1.48130342803 +     1692.16566950240*t);
+   emb_x_0+=     0.00000001132 * cos( 0.96198039943 +    17782.73207278420*t);
+   emb_x_0+=     0.00000000942 * cos( 1.42062620155 +     6267.82337880639*t);
+   emb_x_0+=     0.00000000942 * cos( 5.22790737110 +     6298.32832117640*t);
+   emb_x_0+=     0.00000001045 * cos( 0.46337213067 +    15508.61512327440*t);
+   emb_x_0+=     0.00000000850 * cos( 1.45204280805 +     6393.28217121080*t);
+   emb_x_0+=     0.00000000895 * cos( 2.40085098648 +     3903.91137641980*t);
+   emb_x_0+=     0.00000001016 * cos( 0.71504248238 +      220.41264243880*t);
+   emb_x_0+=     0.00000000853 * cos( 4.42455068314 +      949.17560896980*t);
+   emb_x_0+=     0.00000001085 * cos( 1.32862144392 +    36949.23080842420*t);
+   emb_x_0+=     0.00000000779 * cos( 3.58339451743 +     2648.45482547300*t);
+   emb_x_0+=     0.00000001000 * cos( 0.24880432000 +    11403.67699557500*t);
+   emb_x_0+=     0.00000000962 * cos( 3.92658368218 +    12567.13202105100*t);
+   emb_x_0+=     0.00000000770 * cos( 6.20173754399 +     4804.20927592700*t);
+   emb_x_0+=     0.00000000986 * cos( 4.62897469531 +    22805.73556599360*t);
+   emb_x_0+=     0.00000000931 * cos( 2.22558469426 +     2787.04302385740*t);
+   emb_x_0+=     0.00000000743 * cos( 3.79903933428 +     7834.12107263940*t);
+   emb_x_0+=     0.00000000905 * cos( 2.83314540925 +     6281.59137728310*t);
+   emb_x_0+=     0.00000000905 * cos( 3.81538816339 +     6284.56032269970*t);
+   emb_x_0+=     0.00000000715 * cos( 5.50312651042 +    17256.63153634140*t);
+   emb_x_0+=     0.00000000798 * cos( 2.62501077543 +    12410.73130054860*t);
+   emb_x_0+=     0.00000000828 * cos( 1.54060095214 +    10969.96525769820*t);
+   emb_x_0+=     0.00000000821 * cos( 0.52122150189 +    12416.58850284820*t);
+   emb_x_0+=     0.00000000826 * cos( 4.78071938299 +    10984.19235169980*t);
+   emb_x_0+=     0.00000000754 * cos( 3.82898538792 +     4061.21921539440*t);
+   emb_x_0+=     0.00000000616 * cos( 5.70479516739 +     1596.18644228460*t);
+   emb_x_0+=     0.00000000782 * cos( 0.87676753080 +     4907.30205014560*t);
+   emb_x_0+=     0.00000000707 * cos( 3.28533306159 +     5849.36411211460*t);
+   emb_x_0+=     0.00000000783 * cos( 3.43481557690 +      846.08283475120*t);
+   emb_x_0+=     0.00000000623 * cos( 2.94730498294 +     5333.90024102160*t);
+   emb_x_0+=     0.00000000696 * cos( 4.89623582044 +    11712.95531823080*t);
+   emb_x_0+=     0.00000000595 * cos( 1.30546777768 +    12323.42309600880*t);
+   emb_x_0+=     0.00000000611 * cos( 1.31321200457 +    11933.36796066960*t);
+   emb_x_0+=     0.00000000611 * cos( 1.53155679450 +      853.19638175200*t);
+   emb_x_0+=     0.00000000570 * cos( 4.54665961586 +    12586.92709547520*t);
+   emb_x_0+=     0.00000000778 * cos( 1.57590278073 +     8662.24032356300*t);
+   emb_x_0+=     0.00000000597 * cos( 4.69316038420 +     9779.10867612540*t);
+   emb_x_0+=     0.00000000536 * cos( 6.06660590951 +    18852.75066832320*t);
+   emb_x_0+=     0.00000000532 * cos( 4.58762153412 +    18842.11400297339*t);
+   emb_x_0+=     0.00000000599 * cos( 2.11793250999 +     6151.53388830500*t);
+   emb_x_0+=     0.00000000540 * cos( 3.40440707564 +     6279.78949257360*t);
+   emb_x_0+=     0.00000000540 * cos( 3.24412649701 +     6286.36220740920*t);
+   emb_x_0+=     0.00000000621 * cos( 1.14533482552 +    52175.80628314840*t);
+   emb_x_0+=     0.00000000513 * cos( 3.70291605445 +    12455.94537876340*t);
+   emb_x_0+=     0.00000000623 * cos( 5.52546598759 +     5326.78669402080*t);
+   emb_x_0+=     0.00000000649 * cos( 4.90568768794 +    11764.33076885900*t);
+   emb_x_0+=     0.00000000582 * cos( 4.53092407532 +     6414.61781167780*t);
+   emb_x_0+=     0.00000000610 * cos( 3.94320832220 +    10419.98628350760*t);
+   emb_x_0+=     0.00000000679 * cos( 0.26045226696 +     6133.51265285680*t);
+   emb_x_0+=     0.00000000531 * cos( 3.00684541494 +     3496.03282613400*t);
+   emb_x_0+=     0.00000000536 * cos( 3.70224967645 +     7232.25145896120*t);
+   emb_x_0+=     0.00000000510 * cos( 4.38198402315 +    11513.88331679440*t);
+   emb_x_0+=     0.00000000517 * cos( 2.93096521231 +      433.71173787680*t);
+   emb_x_0+=     0.00000000643 * cos( 6.06382589060 +    26735.94526221320*t);
+   emb_x_0+=     0.00000000665 * cos( 2.75085842376 +    40879.44050464380*t);
+   emb_x_0+=     0.00000000549 * cos( 0.63911798007 +    22483.84857449259*t);
+   emb_x_0+=     0.00000000479 * cos( 1.33580595339 +     5547.19933645960*t);
+   emb_x_0+=     0.00000000564 * cos( 3.92025647347 +     4171.42553661380*t);
+   emb_x_0+=     0.00000000462 * cos( 1.96611775944 +     3104.93005942380*t);
+   emb_x_0+=     0.00000000551 * cos( 4.03874430699 +     6503.48849243020*t);
+   emb_x_0+=     0.00000000545 * cos( 2.61441075819 +     6062.66320755260*t);
+   emb_x_0+=     0.00000000477 * cos( 3.50657951584 +     9388.00590941520*t);
+   emb_x_0+=     0.00000000448 * cos( 5.09815034585 +     4274.51831083240*t);
+   emb_x_0+=     0.00000000532 * cos( 0.22623731981 +      283.85931886520*t);
+   emb_x_0+=     0.00000000527 * cos( 6.24827316267 +    24602.61243487099*t);
+   emb_x_0+=     0.00000000502 * cos( 0.97407194728 +    12489.88562870720*t);
+   emb_x_0+=     0.00000000493 * cos( 6.11920290350 +    29296.61538957860*t);
+   emb_x_0+=     0.00000000530 * cos( 0.43491203890 +    18845.70443162520*t);
+   emb_x_0+=     0.00000000431 * cos( 5.31145947560 +     7018.95236352320*t);
+   emb_x_0+=     0.00000000559 * cos( 1.71609589414 +     1581.95934828300*t);
+   emb_x_0+=     0.00000000478 * cos( 1.85248622079 +       14.22709400160*t);
+   emb_x_0+=     0.00000000405 * cos( 5.39493918596 +    10873.98603048040*t);
+   emb_x_0+=     0.00000000401 * cos( 4.50443000854 +    17267.26820169119*t);
+   emb_x_0+=     0.00000000389 * cos( 3.21806007009 +     6321.20888562920*t);
+   emb_x_0+=     0.00000000389 * cos( 3.43047350256 +     6244.94281435360*t);
+   emb_x_0+=     0.00000000388 * cos( 1.13963224225 +     6297.30294399300*t);
+   emb_x_0+=     0.00000000388 * cos( 5.50890133040 +     6268.84875598980*t);
+   emb_x_0+=     0.00000000515 * cos( 1.67075193318 +    12320.32005375340*t);
+   emb_x_0+=     0.00000000393 * cos( 4.37184659923 +    14919.01785375460*t);
+   emb_x_0+=     0.00000000394 * cos( 6.10411050787 +     4157.19844261220*t);
+   emb_x_0+=     0.00000000491 * cos( 5.96740803508 +     1478.86657406440*t);
+   emb_x_0+=     0.00000000505 * cos( 4.66096327599 +    17157.06188047180*t);
+   emb_x_0+=     0.00000000528 * cos( 5.62841514056 +      135.06508003540*t);
+   emb_x_0+=     0.00000000417 * cos( 3.74443512826 +    13760.59871020740*t);
+   emb_x_0+=     0.00000000419 * cos( 5.08318905121 +    15671.08175940660*t);
+   emb_x_0+=     0.00000000393 * cos( 1.56711153414 +     7872.14874527520*t);
+   emb_x_0+=     0.00000000392 * cos( 1.30305906278 +      419.48464387520*t);
+   emb_x_0+=     0.00000000365 * cos( 2.11020607854 +     6922.97313630540*t);
+   emb_x_0+=     0.00000000419 * cos( 0.62417554081 +     5642.19824260920*t);
+   emb_x_0+=     0.00000000387 * cos( 0.68719040069 +    17253.04110768959*t);
+   emb_x_0+=     0.00000000371 * cos( 5.76233828697 +    13625.53363017200*t);
+   emb_x_0+=     0.00000000379 * cos( 0.81831291068 +     3185.19202726560*t);
+   emb_x_0+=     0.00000000414 * cos( 3.32204347851 +    22779.43724619380*t);
+   emb_x_0+=     0.00000000485 * cos( 1.35444198321 +    12528.01866434500*t);
+   emb_x_0+=     0.00000000358 * cos( 3.15737999738 +     6271.34649715540*t);
+   emb_x_0+=     0.00000000358 * cos( 3.49115357527 +     6294.80520282740*t);
+   emb_x_0+=     0.00000000351 * cos( 2.21887707287 +    11616.97609101300*t);
+   emb_x_0+=     0.00000000387 * cos( 0.49044025161 +      956.28915597060*t);
+   emb_x_0+=     0.00000000451 * cos( 0.91836677580 +    29864.33402730900*t);
+   emb_x_0+=     0.00000000421 * cos( 4.82745120825 +     7445.55055439920*t);
+   emb_x_0+=     0.00000000327 * cos( 3.95232223910 +    20995.39296644940*t);
+   emb_x_0+=     0.00000000336 * cos( 5.23522689515 +     6287.00800325450*t);
+   emb_x_0+=     0.00000000336 * cos( 1.41330667750 +     6279.14369672830*t);
+   emb_x_0+=     0.00000000358 * cos( 0.81490874525 +    21393.54196985760*t);
+   emb_x_0+=     0.00000000320 * cos( 0.03172757569 +    16858.48253293320*t);
+   emb_x_0+=     0.00000000348 * cos( 4.27565702309 +        3.93215326310*t);
+   emb_x_0+=     0.00000000363 * cos( 4.29069778349 +     7335.34423317980*t);
+   emb_x_0+=     0.00000000330 * cos( 3.68113867113 +       38.02767263580*t);
+   emb_x_0+=     0.00000000395 * cos( 5.39428043096 +     5905.70224207560*t);
+   emb_x_0+=     0.00000000414 * cos( 4.19402575890 +    44809.65020086340*t);
+   emb_x_0+=     0.00000000300 * cos( 3.87083875051 +      640.87760738220*t);
+   emb_x_0+=     0.00000000304 * cos( 0.15986036979 +     5749.86176654780*t);
+   emb_x_0+=     0.00000000351 * cos( 5.93395455459 +    12345.73905754400*t);
+   emb_x_0+=     0.00000000307 * cos( 5.30261033229 +     5120.60114558360*t);
+   emb_x_0+=     0.00000000304 * cos( 4.93657040735 +    29826.30635467320*t);
+   emb_x_0+=     0.00000000297 * cos( 6.01664889580 +    20597.24396304120*t);
+   emb_x_0+=     0.00000000395 * cos( 0.07726910331 +     4176.04134244900*t);
+   emb_x_0+=     0.00000000320 * cos( 2.49799253183 +    10818.13528691580*t);
+   emb_x_0+=     0.00000000365 * cos( 1.25425297875 +     6660.44945790720*t);
+   emb_x_0+=     0.00000000339 * cos( 6.10828911901 +     6923.95345737360*t);
+   emb_x_0+=     0.00000000274 * cos( 1.11162071728 +    26709.64694241340*t);
+   emb_x_0+=     0.00000000386 * cos( 2.29044623019 +     2636.72547263700*t);
+   emb_x_0+=     0.00000000367 * cos( 1.18722914434 +    30666.15495843280*t);
+   emb_x_0+=     0.00000000268 * cos( 4.95835137530 +    10454.50138660520*t);
+   emb_x_0+=     0.00000000265 * cos( 5.11010291991 +    19378.91851506880*t);
+   emb_x_0+=     0.00000000325 * cos( 2.49277998223 +    12359.96615154560*t);
+   emb_x_0+=     0.00000000326 * cos( 4.24702233173 +    16522.65971600220*t);
+   emb_x_0+=     0.00000000269 * cos( 0.20537129126 +     6816.28993343500*t);
+   emb_x_0+=     0.00000000298 * cos( 5.15059518785 +    10344.29506538580*t);
+   emb_x_0+=     0.00000000344 * cos( 3.51559745333 +     6147.45052498140*t);
+   emb_x_0+=     0.00000000255 * cos( 4.94642831875 +     1375.77379984580*t);
+   emb_x_0+=     0.00000000283 * cos( 2.57877323959 +      316.39186965660*t);
+   emb_x_0+=     0.00000000334 * cos( 3.13303102229 +     6418.70117500140*t);
+   emb_x_0+=     0.00000000316 * cos( 5.54205366730 +     6055.59971720240*t);
+   emb_x_0+=     0.00000000316 * cos( 0.36194124996 +    21954.15760939799*t);
+   emb_x_0+=     0.00000000243 * cos( 4.10172883526 +     6310.47740608820*t);
+   emb_x_0+=     0.00000000243 * cos( 2.54680473739 +     6255.67429389460*t);
+   emb_x_0+=     0.00000000314 * cos( 4.99359717889 +    17996.03116822220*t);
+   emb_x_0+=     0.00000000296 * cos( 0.07899534657 +     6129.29703950660*t);
+   emb_x_0+=     0.00000000301 * cos( 1.10648000281 +     6510.55198278040*t);
+   emb_x_0+=     0.00000000329 * cos( 1.15676558971 +      377.37360791580*t);
+   emb_x_0+=     0.00000000267 * cos( 3.42007125342 +    24705.70520908959*t);
+   emb_x_0+=     0.00000000285 * cos( 2.50494681053 +     5469.52556603160*t);
+   emb_x_0+=     0.00000000311 * cos( 3.23911937497 +    11720.06886523160*t);
+   emb_x_0+=     0.00000000286 * cos( 0.28647339668 +     6436.85466047620*t);
+   emb_x_0+=     0.00000000273 * cos( 4.29648399498 +    28286.99048486120*t);
+   emb_x_0+=     0.00000000229 * cos( 1.01384122818 +     1265.56747862640*t);
+   emb_x_0+=     0.00000000234 * cos( 3.36893910039 +     4487.81740627040*t);
+   emb_x_0+=     0.00000000252 * cos( 0.51301411030 +     4377.61108505100*t);
+   emb_x_0+=     0.00000000214 * cos( 2.09487711967 +     2221.85663459700*t);
+   emb_x_0+=     0.00000000213 * cos( 3.79431896643 +    30355.99731976780*t);
+   emb_x_0+=     0.00000000263 * cos( 6.11317571912 +     9380.95967271720*t);
+   emb_x_0+=     0.00000000239 * cos( 1.33165763648 +    10440.27429260360*t);
+   emb_x_0+=     0.00000000289 * cos( 0.55454156939 +     6206.80977871580*t);
+   emb_x_0+=     0.00000000259 * cos( 3.43127382595 +     6282.83203250789*t);
+   emb_x_0+=     0.00000000259 * cos( 3.21725974670 +     6283.31966747490*t);
+   emb_x_0+=     0.00000000243 * cos( 4.14346390349 +     7096.62613395120*t);
+   emb_x_0+=     0.00000000246 * cos( 5.83734028861 +     8671.96987044060*t);
+   emb_x_0+=     0.00000000199 * cos( 0.93488663037 +    11216.28429032400*t);
+   emb_x_0+=     0.00000000265 * cos( 5.61985831637 +    48739.85989708300*t);
+   emb_x_0+=     0.00000000207 * cos( 3.06628694399 +    12562.56127133100*t);
+   emb_x_0+=     0.00000000212 * cos( 1.75471606720 +     7129.15868474260*t);
+   emb_x_0+=     0.00000000225 * cos( 4.01656879083 +     5636.06501667660*t);
+   emb_x_0+=     0.00000000192 * cos( 2.04748658291 +    17271.88400752640*t);
+   emb_x_0+=     0.00000000240 * cos( 5.09605084016 +    15906.76412668260*t);
+   emb_x_0+=     0.00000000187 * cos( 0.15192842664 +    11609.86254401220*t);
+   emb_x_0+=     0.00000000256 * cos( 2.30143585875 +      224.34479570190*t);
+   emb_x_0+=     0.00000000250 * cos( 4.24183292428 +     5934.15142954340*t);
+   emb_x_0+=     0.00000000220 * cos( 0.73345029623 +    11087.28512591840*t);
+   emb_x_0+=     0.00000000175 * cos( 5.41174539523 +    11614.43329373220*t);
+   emb_x_0+=     0.00000000175 * cos( 4.58212579300 +     3744.82734573780*t);
+   emb_x_0+=     0.00000000241 * cos( 0.28788321480 +     8390.11035753380*t);
+   emb_x_0+=     0.00000000234 * cos( 2.61192358068 +    34596.36465465240*t);
+   emb_x_0+=     0.00000000174 * cos( 2.85207900960 +      263.08392337280*t);
+   emb_x_0+=     0.00000000233 * cos( 2.02657287745 +        1.48447270830*t);
+   emb_x_0+=     0.00000000190 * cos( 6.26917436708 +      394.62588505920*t);
+   emb_x_0+=     0.00000000187 * cos( 5.58825382870 +      401.67212175720*t);
+   emb_x_0+=     0.00000000178 * cos( 3.72578587318 +    12132.43996210600*t);
+   emb_x_0+=     0.00000000230 * cos( 2.40692341132 +     6632.00027043940*t);
+   emb_x_0+=     0.00000000177 * cos( 4.73089168144 +     4797.09572892620*t);
+   emb_x_0+=     0.00000000167 * cos( 3.97388034335 +     5714.25397596400*t);
+   emb_x_0+=     0.00000000160 * cos( 2.37187295106 +      742.99006053260*t);
+   emb_x_0+=     0.00000000214 * cos( 4.34116392416 +    24492.40611365159*t);
+   emb_x_0+=     0.00000000217 * cos( 4.52620284574 +     3097.88382272579*t);
+   emb_x_0+=     0.00000000172 * cos( 5.38973312367 +     6332.83287546320*t);
+   emb_x_0+=     0.00000000172 * cos( 1.25880044898 +     6233.31882451960*t);
+   emb_x_0+=     0.00000000221 * cos( 4.98998061605 +      266.60704172180*t);
+   emb_x_0+=     0.00000000159 * cos( 3.90246086071 +    13198.93543929600*t);
+   emb_x_0+=     0.00000000207 * cos( 3.49598218672 +    12569.74212863460*t);
+   emb_x_0+=     0.00000000176 * cos( 4.05199455874 +    12012.58229714040*t);
+   emb_x_0+=     0.00000000157 * cos( 5.43302098132 +    16840.67001081519*t);
+   emb_x_0+=     0.00000000170 * cos( 0.63813817825 +     6272.03014972750*t);
+   emb_x_0+=     0.00000000170 * cos( 6.01039539440 +     6294.12155025530*t);
+   emb_x_0+=     0.00000000213 * cos( 4.21609963122 +    19804.82729158280*t);
+   emb_x_0+=     0.00000000167 * cos( 5.54070880349 +    17686.75284556640*t);
+   emb_x_0+=     0.00000000187 * cos( 5.70099962642 +    20452.86941222180*t);
+   emb_x_0+=     0.00000000147 * cos( 4.99188388936 +    16737.57723659660*t);
+   emb_x_0+=     0.00000000148 * cos( 2.67465318849 +     6851.89772401880*t);
+   emb_x_0+=     0.00000000146 * cos( 0.23801153618 +    17370.36097590980*t);
+   emb_x_0+=     0.00000000143 * cos( 0.85202449839 +     5618.31980486140*t);
+   emb_x_0+=     0.00000000193 * cos( 3.57442994352 +    33326.57873317420*t);
+   emb_x_0+=     0.00000000140 * cos( 1.22261707467 +     2008.55753915900*t);
+   emb_x_0+=     0.00000000155 * cos( 0.27598891874 +    20199.09495963300*t);
+   emb_x_0+=     0.00000000182 * cos( 5.69212378196 +    28766.92442448400*t);
+   emb_x_0+=     0.00000000184 * cos( 3.48483358167 +     5540.08578945880*t);
+   emb_x_0+=     0.00000000162 * cos( 4.78754155685 +    12551.92460598120*t);
+   emb_x_0+=     0.00000000143 * cos( 2.32151913243 +     5010.39482436420*t);
+   emb_x_0+=     0.00000000190 * cos( 0.66924538276 +    23581.25817731760*t);
+   emb_x_0+=     0.00000000135 * cos( 5.34929799142 +      735.87651353180*t);
+   emb_x_0+=     0.00000000135 * cos( 0.07815962516 +    12431.08661994740*t);
+   emb_x_0+=     0.00000000134 * cos( 3.03203004875 +    11925.27409260060*t);
+   emb_x_0+=     0.00000000179 * cos( 1.07503153366 +    16460.33352952499*t);
+   emb_x_0+=     0.00000000182 * cos( 3.66470209346 +       28.44918746780*t);
+   emb_x_0+=     0.00000000146 * cos( 1.55848605346 +    12701.21678001820*t);
+   emb_x_0+=     0.00000000137 * cos( 2.60133025404 +       11.04570026390*t);
+   emb_x_0+=     0.00000000167 * cos( 1.26986508714 +    17473.45375012840*t);
+   emb_x_0+=     0.00000000170 * cos( 0.76594212114 +    52670.06959330260*t);
+   emb_x_0+=     0.00000000124 * cos( 5.79650915048 +     6947.83189512140*t);
+   emb_x_0+=     0.00000000135 * cos( 3.38698681678 +     3531.52825029980*t);
+   emb_x_0+=     0.00000000122 * cos( 5.69264255527 +     3583.34103067380*t);
+   emb_x_0+=     0.00000000155 * cos( 4.05744056379 +     6058.73105428950*t);
+   emb_x_0+=     0.00000000120 * cos( 1.19140907580 +    12808.88030395680*t);
+   emb_x_0+=     0.00000000124 * cos( 1.37398060221 +     8982.81066930900*t);
+   emb_x_0+=     0.00000000120 * cos( 5.84076902832 +    12528.12402734700*t);
+   emb_x_0+=     0.00000000119 * cos( 1.19695856457 +    19247.37655338240*t);
+   emb_x_0+=     0.00000000143 * cos( 2.13563488838 +       95.97922721780*t);
+   emb_x_0+=     0.00000000117 * cos( 0.73795877914 +    11610.55195837420*t);
+   emb_x_0+=     0.00000000154 * cos( 4.06475922223 +    38526.57435087200*t);
+   emb_x_0+=     0.00000000163 * cos( 3.11560982847 +      661.23292678100*t);
+   emb_x_0+=     0.00000000163 * cos( 4.35848910377 +     9929.42622734580*t);
+   emb_x_0+=     0.00000000115 * cos( 2.14676232906 +    19402.79695281660*t);
+   emb_x_0+=     0.00000000119 * cos( 4.28708216617 +    12249.75983032620*t);
+   emb_x_0+=     0.00000000145 * cos( 2.74250030398 +     6360.74962041940*t);
+   emb_x_0+=     0.00000000145 * cos( 3.90603326867 +     6205.40207956340*t);
+   emb_x_0+=     0.00000000156 * cos( 2.57788957454 +     4307.58330413540*t);
+   emb_x_0+=     0.00000000116 * cos( 0.60146141700 +    24065.80792277559*t);
+   emb_x_0+=     0.00000000118 * cos( 1.23925579836 +     6302.74561089119*t);
+   emb_x_0+=     0.00000000118 * cos( 5.40927777429 +     6263.40608909160*t);
+   emb_x_0+=     0.00000000110 * cos( 2.33480621176 +    16062.18452611680*t);
+   emb_x_0+=     0.00000000137 * cos( 2.54415078163 +     6315.60840078280*t);
+   emb_x_0+=     0.00000000137 * cos( 4.10438279102 +     6250.54329920000*t);
+   emb_x_0+=     0.00000000111 * cos( 2.27442358208 +     6339.97422492700*t);
+   emb_x_0+=     0.00000000111 * cos( 4.37410999057 +     6226.17747505580*t);
+   emb_x_0+=     0.00000000150 * cos( 0.96205875343 +    19800.94595622480*t);
+   emb_x_0+=     0.00000000123 * cos( 2.37616666242 +    24734.15439655740*t);
+   emb_x_0+=     0.00000000107 * cos( 2.16622176546 +    12463.05892576420*t);
+   emb_x_0+=     0.00000000123 * cos( 1.00594418221 +     6305.85105144220*t);
+   emb_x_0+=     0.00000000123 * cos( 5.64258939044 +     6260.30064854060*t);
+   emb_x_0+=     0.00000000106 * cos( 1.67655703852 +    11190.37790013700*t);
+   emb_x_0+=     0.00000000113 * cos( 1.95413568499 +     2111.65031337760*t);
+   emb_x_0+=     0.00000000131 * cos( 3.61185299316 +    18415.51581209739*t);
+   emb_x_0+=     0.00000000110 * cos( 4.88667127380 +     6276.44729109139*t);
+   emb_x_0+=     0.00000000110 * cos( 1.76186229885 +     6289.70440889141*t);
+   emb_x_0+=     0.00000000103 * cos( 2.75809551865 +     6324.62664097620*t);
+   emb_x_0+=     0.00000000103 * cos( 3.89043805400 +     6241.52505900660*t);
+   emb_x_0+=     0.00000000103 * cos( 5.76083420502 +    23937.85638974100*t);
+   emb_x_0+=     0.00000000113 * cos( 2.65508166498 +    24336.00539314920*t);
+   emb_x_0+=     0.00000000106 * cos( 3.27680652019 +    12772.33724842000*t);
+   emb_x_0+=     0.00000000143 * cos( 1.51199781479 +     2069.00683490660*t);
+   emb_x_0+=     0.00000000143 * cos( 0.91985752497 +    18772.96147869860*t);
+   emb_x_0+=     0.00000000115 * cos( 3.78220330855 +     3641.73457151920*t);
+   emb_x_0+=     0.00000000126 * cos( 1.19294266238 +    18822.92923017440*t);
+   emb_x_0+=     0.00000000115 * cos( 4.66544469064 +    12570.08385324590*t);
+   emb_x_0+=     0.00000000098 * cos( 5.07674655291 +    16310.97904572060*t);
+   emb_x_0+=     0.00000000138 * cos( 2.08935051421 +     6356.37297585040*t);
+   emb_x_0+=     0.00000000138 * cos( 4.55918305844 +     6209.77872413240*t);
+   emb_x_0+=     0.00000000120 * cos( 4.95740945799 +     9814.60410029120*t);
+   emb_x_0+=     0.00000000135 * cos( 0.57059381593 +    18699.66435283960*t);
+   emb_x_0+=     0.00000000104 * cos( 2.15061044327 +    12560.62877567540*t);
+   emb_x_0+=     0.00000000098 * cos( 5.57514710248 +     5815.11085963699*t);
+   emb_x_0+=     0.00000000106 * cos( 1.77919890262 +    23539.70738633280*t);
+   emb_x_0+=     0.00000000114 * cos( 3.75372093481 +     6155.05700665400*t);
+   emb_x_0+=     0.00000000114 * cos( 2.89481263784 +     6411.09469332880*t);
+   emb_x_0+=     0.00000000097 * cos( 6.03825791149 +    27043.50288318280*t);
+   emb_x_0+=     0.00000000095 * cos( 5.63307798614 +    11300.58422135640*t);
+   emb_x_0+=     0.00000000110 * cos( 2.79120257993 +     6312.50535852740*t);
+   emb_x_0+=     0.00000000110 * cos( 3.85733099272 +     6253.64634145540*t);
+   emb_x_0+=     0.00000000091 * cos( 2.99457476326 +    12829.23562335560*t);
+   emb_x_0+=     0.00000000128 * cos( 5.41337210078 +     1293.48422358260*t);
+   emb_x_0+=     0.00000000096 * cos( 5.97584747761 +     2699.73481931760*t);
+   emb_x_0+=     0.00000000119 * cos( 5.98882616741 +    17892.93839400359*t);
+   emb_x_0+=     0.00000000120 * cos( 2.08798412726 +    14158.74771361560*t);
+   emb_x_0+=     0.00000000100 * cos( 4.30344030661 +    16627.37091537720*t);
+   emb_x_0+=     0.00000000099 * cos( 3.61958161583 +     6432.63904712600*t);
+   emb_x_0+=     0.00000000095 * cos( 4.18452410969 +    78263.70942472259*t);
+   emb_x_0+=     0.00000000089 * cos( 1.07338604948 +     6751.04084034580*t);
+   emb_x_0+=     0.00000000096 * cos( 2.64494034094 +     7349.57132718140*t);
+   emb_x_0+=     0.00000000117 * cos( 5.44811650549 +    42456.78404709160*t);
+   emb_x_0+=     0.00000000103 * cos( 4.37748903957 +    18202.21671665939*t);
+   emb_x_0+=     0.00000000099 * cos( 0.58169736416 +     6313.98997562640*t);
+   emb_x_0+=     0.00000000099 * cos( 6.06683620849 +     6252.16172435640*t);
+   emb_x_0+=     0.00000000097 * cos( 0.27231106216 +    10557.59416082380*t);
+   emb_x_0+=     0.00000000101 * cos( 3.42372783781 +    19004.64794940840*t);
+   emb_x_0+=     0.00000000098 * cos( 3.23381285207 +     6126.67512948900*t);
+   emb_x_0+=     0.00000000098 * cos( 3.41472072058 +     6439.47657049380*t);
+   emb_x_0+=     0.00000000086 * cos( 1.34149823513 +    23123.74586080659*t);
+   emb_x_0+=     0.00000000084 * cos( 3.13824782000 +     8858.31494432060*t);
+   emb_x_0+=     0.00000000101 * cos( 0.23009068890 +    18429.74290609900*t);
+   emb_x_0+=     0.00000000088 * cos( 3.76745959676 +    17053.96910625319*t);
+   emb_x_0+=     0.00000000113 * cos( 2.05049922518 +    12985.63634385800*t);
+   emb_x_0+=     0.00000000082 * cos( 0.19375634761 +     5483.25472482600*t);
+   emb_x_0+=     0.00000000084 * cos( 4.23745405632 +     6379.05507720920*t);
+   emb_x_0+=     0.00000000084 * cos( 2.41107951633 +     6187.09662277360*t);
+   emb_x_0+=     0.00000000097 * cos( 4.02159784443 +    10239.58386601080*t);
+   emb_x_0+=     0.00000000105 * cos( 2.77478785571 +     6434.97313107660*t);
+   emb_x_0+=     0.00000000105 * cos( 3.87374571694 +     6131.17856890620*t);
+   emb_x_0+=     0.00000000114 * cos( 4.07080498725 +     8258.56839584740*t);
+   emb_x_0+=     0.00000000082 * cos( 3.82704107824 +     3958.12644117580*t);
+   emb_x_0+=     0.00000000093 * cos( 1.35189028328 +    18216.44381066100*t);
+   emb_x_0+=     0.00000000102 * cos( 2.89783593458 +     5767.61197889840*t);
+   emb_x_0+=     0.00000000102 * cos( 3.75069763806 +     6798.53972108440*t);
+   emb_x_0+=     0.00000000110 * cos( 2.19234510205 +    56600.27928952220*t);
+   emb_x_0+=     0.00000000112 * cos( 5.48905138227 +    12188.77809206700*t);
+   emb_x_0+=     0.00000000080 * cos( 3.59070813227 +    24080.03501677719*t);
+   emb_x_0+=     0.00000000088 * cos( 3.00398297222 +     1573.82042409880*t);
+   emb_x_0+=     0.00000000087 * cos( 2.32187773582 +     1580.86666079680*t);
+   emb_x_0+=     0.00000000099 * cos( 2.40898854628 +    33794.54372352860*t);
+   emb_x_0+=     0.00000000080 * cos( 3.27906531246 +     6485.32924516550*t);
+   emb_x_0+=     0.00000000080 * cos( 3.36946826019 +     6080.82245481730*t);
+   emb_x_0+=     0.00000000089 * cos( 5.20240793768 +    17363.24742890899*t);
+   emb_x_0+=     0.00000000079 * cos( 4.53817935583 +    36109.38220466459*t);
+   emb_x_0+=     0.00000000088 * cos( 0.87070535180 +    16943.76278503380*t);
+   emb_x_0+=     0.00000000088 * cos( 3.24901264480 +    12537.70251251500*t);
+   emb_x_0+=     0.00000000087 * cos( 1.75120977431 +    28237.23345938940*t);
+   emb_x_0+=     0.00000000080 * cos( 3.62211320521 +     3956.50801601940*t);
+   emb_x_0+=     0.00000000083 * cos( 0.01677699839 +    21424.46664430340*t);
+   emb_x_0+=     0.00000000099 * cos( 4.35778148149 +     5017.50837136500*t);
+   emb_x_0+=     0.00000000079 * cos( 4.50777133706 +     5621.84292321040*t);
+   emb_x_0+=     0.00000000096 * cos( 1.98497858569 +     5888.44996493220*t);
+   emb_x_0+=     0.00000000096 * cos( 4.66355498696 +     6677.70173505060*t);
+   emb_x_0+=     0.00000000074 * cos( 5.98124271909 +     6304.41149045840*t);
+   emb_x_0+=     0.00000000074 * cos( 0.66729085356 +     6261.74020952440*t);
+   emb_x_0+=     0.00000000073 * cos( 2.62129339640 +     6695.44694686580*t);
+   emb_x_0+=     0.00000000073 * cos( 4.02724017625 +     5870.70475311700*t);
+   emb_x_0+=     0.00000000085 * cos( 3.65241534932 +    12850.01101884800*t);
+   emb_x_0+=     0.00000000076 * cos( 1.33659857866 +    40796.33892267419*t);
+   emb_x_0+=     0.00000000097 * cos( 0.92251430623 +    24383.07910844140*t);
+   emb_x_0+=     0.00000000072 * cos( 0.32356695638 +     6312.89728814020*t);
+   emb_x_0+=     0.00000000072 * cos( 0.04178130909 +     6253.25441184259*t);
+   emb_x_0+=     0.00000000095 * cos( 0.62289841825 +    11080.17157891760*t);
+   emb_x_0+=     0.00000000100 * cos( 5.37794668456 +     5316.10497255580*t);
+   emb_x_0+=     0.00000000085 * cos( 2.90983978106 +    12282.29238111760*t);
+   emb_x_0+=     0.00000000071 * cos( 6.13573135109 +    41194.48792608240*t);
+   emb_x_0+=     0.00000000072 * cos( 4.76869158882 +    15141.39079431200*t);
+   emb_x_0+=     0.00000000073 * cos( 1.29002951973 +    45892.73043315699*t);
+   emb_x_0+=     0.00000000076 * cos( 1.98229485596 +    63658.87775083760*t);
+   emb_x_0+=     0.00000000087 * cos( 3.63418769958 +    18106.23748944160*t);
+   emb_x_0+=     0.00000000070 * cos( 2.01720588184 +     6442.01936777460*t);
+   emb_x_0+=     0.00000000070 * cos( 4.63132769081 +     6124.13233220820*t);
+   emb_x_0+=     0.00000000097 * cos( 4.76857001908 +    92747.68916682259*t);
+   emb_x_0+=     0.00000000078 * cos( 4.94336223185 +    23550.34405168259*t);
+   emb_x_0+=     0.00000000076 * cos( 3.43423122344 +     5237.92101380380*t);
+   emb_x_0+=     0.00000000078 * cos( 4.98905027033 +     9924.81042151060*t);
+   emb_x_0+=     0.00000000082 * cos( 5.48703292615 +     6774.73914245020*t);
+   emb_x_0+=     0.00000000067 * cos( 0.17045052997 +     7082.89697515680*t);
+   emb_x_0+=     0.00000000067 * cos( 4.81312573338 +     6171.88920770380*t);
+   emb_x_0+=     0.00000000067 * cos( 1.83540783927 +     6394.26249227900*t);
+   emb_x_0+=     0.00000000080 * cos( 5.36760004078 +    22380.75580027400*t);
+   emb_x_0+=     0.00000000070 * cos( 2.04678665255 +     6175.41232605280*t);
+   emb_x_0+=     0.00000000070 * cos( 4.60174692010 +     6390.73937393000*t);
+   emb_x_0+=     0.00000000091 * cos( 2.47526506655 +     6652.77566593180*t);
+   emb_x_0+=     0.00000000091 * cos( 4.17326850610 +     5913.37603405100*t);
+   emb_x_0+=     0.00000000067 * cos( 0.13241010748 +     6300.88837210940*t);
+   emb_x_0+=     0.00000000067 * cos( 0.23293815799 +     6265.26332787340*t);
+   emb_x_0+=     0.00000000067 * cos( 1.43864813435 +     5959.57043333400*t);
+   emb_x_0+=     0.00000000080 * cos( 0.77942675353 +    23646.32327890039*t);
+   emb_x_0+=     0.00000000068 * cos( 3.64312434743 +     4897.18057365520*t);
+   emb_x_0+=     0.00000000069 * cos( 2.14076204045 +     6944.30877677240*t);
+   emb_x_0+=     0.00000000089 * cos( 1.97945887866 +    10660.68693504240*t);
+   emb_x_0+=     0.00000000073 * cos( 4.53414425046 +     6195.76764545159*t);
+   emb_x_0+=     0.00000000073 * cos( 2.11438932219 +     6370.38405453120*t);
+   emb_x_0+=     0.00000000066 * cos( 4.95039391556 +      323.50541665740*t);
+   emb_x_0+=     0.00000000066 * cos( 3.95274142218 +     6251.10354417460*t);
+   emb_x_0+=     0.00000000066 * cos( 2.69579215047 +     6315.04815580820*t);
+   emb_x_0+=     0.00000000063 * cos( 4.27995423229 +    11823.16163945020*t);
+   emb_x_0+=     0.00000000064 * cos( 0.09094709813 +    28774.03797148480*t);
+   emb_x_0+=     0.00000000086 * cos( 4.32355291567 +    12341.80690428090*t);
+   emb_x_0+=     0.00000000081 * cos( 0.47769423835 +    31441.67756975680*t);
+   emb_x_0+=     0.00000000071 * cos( 5.70281660725 +    16207.88627150200*t);
+   emb_x_0+=     0.00000000061 * cos( 0.43078256287 +    35050.00027447539*t);
+   emb_x_0+=     0.00000000066 * cos( 2.78136130877 +    17576.54652434700*t);
+   emb_x_0+=     0.00000000072 * cos( 0.54355547965 +    18100.00325845000*t);
+   emb_x_0+=     0.00000000064 * cos( 1.93356156446 +    23006.42599258639*t);
+   emb_x_0+=     0.00000000068 * cos( 2.90815872826 +    28759.81087748319*t);
+   emb_x_0+=     0.00000000062 * cos( 0.67966053291 +     3855.03366695720*t);
+   emb_x_0+=     0.00000000073 * cos( 6.18334008725 +    23440.13773046320*t);
+   emb_x_0+=     0.00000000060 * cos( 3.22950754418 +    13088.72911807660*t);
+   emb_x_0+=     0.00000000078 * cos( 4.56094059497 +    12564.66722727450*t);
+   emb_x_0+=     0.00000000060 * cos( 4.47214922427 +     7548.64332861780*t);
+   emb_x_0+=     0.00000000075 * cos( 4.42126467141 +    18624.88275427230*t);
+   emb_x_0+=     0.00000000059 * cos( 0.42392189557 +    41990.78593289880*t);
+   emb_x_0+=     0.00000000081 * cos( 1.27058672695 +     7250.04672742700*t);
+   emb_x_0+=     0.00000000059 * cos( 3.96273011871 +     6289.14176562120*t);
+   emb_x_0+=     0.00000000059 * cos( 2.68580345394 +     6277.00993436160*t);
+   emb_x_0+=     0.00000000063 * cos( 4.91003104939 +     5436.99301524020*t);
+   emb_x_0+=     0.00000000059 * cos( 2.73782977803 +     8584.66166590080*t);
+   emb_x_0+=     0.00000000074 * cos( 2.29556008890 +    28313.28880466100*t);
+   emb_x_0+=     0.00000000057 * cos( 0.01221429824 +    10027.90319572920*t);
+   emb_x_0+=     0.00000000057 * cos( 4.87299369650 +      792.77488846740*t);
+   emb_x_0+=     0.00000000060 * cos( 2.20527423485 +      647.01083331480*t);
+   emb_x_0+=     0.00000000056 * cos( 4.47895058873 +    15265.88651930040*t);
+   emb_x_0+=     0.00000000072 * cos( 4.16398251563 +     4598.02372748980*t);
+   emb_x_0+=     0.00000000063 * cos( 1.27506758028 +        5.52292430740*t);
+   emb_x_0+=     0.00000000068 * cos( 2.71579728238 +     3166.41643773160*t);
+   emb_x_0+=     0.00000000074 * cos( 3.23404191447 +    82534.40362061159*t);
+   emb_x_0+=     0.00000000056 * cos( 1.96997381249 +     6339.87911216120*t);
+   emb_x_0+=     0.00000000056 * cos( 4.67855976016 +     6226.27258782160*t);
+   emb_x_0+=     0.00000000069 * cos( 1.89738325082 +    23536.11695768099*t);
+   emb_x_0+=     0.00000000056 * cos( 2.52728721608 +    12303.06777661000*t);
+   emb_x_0+=     0.00000000055 * cos( 3.89144240866 +    22345.26037610820*t);
+   emb_x_0+=     0.00000000057 * cos( 3.55801157684 +     6702.00024889200*t);
+   emb_x_0+=     0.00000000057 * cos( 3.09052199581 +     5864.15145109080*t);
+   emb_x_0+=     0.00000000064 * cos( 5.42433316705 +    14169.79356223040*t);
+   emb_x_0+=     0.00000000075 * cos( 5.93374782872 +    32217.20018108080*t);
+   emb_x_0+=     0.00000000059 * cos( 2.17826568386 +    18003.14471522300*t);
+   emb_x_0+=     0.00000000059 * cos( 3.38504946794 +    67589.08744705719*t);
+   emb_x_0+=     0.00000000076 * cos( 2.34884503279 +    23226.83863502520*t);
+   emb_x_0+=     0.00000000073 * cos( 6.00950734574 +    30348.88377276700*t);
+   emb_x_0+=     0.00000000070 * cos( 1.56579326622 +    12299.54465826100*t);
+   emb_x_0+=     0.00000000060 * cos( 0.37221954582 +    12171.52581492360*t);
+   emb_x_0+=     0.00000000058 * cos( 3.15618940951 +     6716.78758786820*t);
+   emb_x_0+=     0.00000000074 * cos( 5.91246929008 +     3684.37804999020*t);
+   emb_x_0+=     0.00000000067 * cos( 2.65177830926 +    22743.40937951640*t);
+   emb_x_0+=     0.00000000062 * cos( 2.47017142882 +    16193.65917750039*t);
+   emb_x_0+=     0.00000000059 * cos( 2.60023110979 +       15.25247118500*t);
+   emb_x_0+=     0.00000000059 * cos( 1.05069387628 +    12164.47957822560*t);
+   emb_x_0+=     0.00000000056 * cos( 4.02457745520 +     3178.14579056760*t);
+   emb_x_0+=     0.00000000063 * cos( 3.66056370325 +     5219.76176653910*t);
+   emb_x_0+=     0.00000000063 * cos( 0.58993036842 +     5227.62607306530*t);
+   emb_x_0+=     0.00000000071 * cos( 3.61897462493 +    60530.48898574180*t);
+   emb_x_0+=     0.00000000072 * cos( 1.05408545034 +     4583.79663348820*t);
+   emb_x_0+=     0.00000000052 * cos( 5.67741459488 +     3627.50747751760*t);
+   emb_x_0+=     0.00000000052 * cos( 1.77495138030 +     6155.12431695680*t);
+   emb_x_0+=     0.00000000052 * cos( 4.87358219235 +     6411.02738302600*t);
+   emb_x_0+=     0.00000000073 * cos( 4.40543588100 +    80181.53746683979*t);
+   emb_x_0+=     0.00000000058 * cos( 4.29308512080 +    30775.48196364300*t);
+   emb_x_0+=     0.00000000068 * cos( 1.86654756299 +    35371.88726597640*t);
+   emb_x_0+=     0.00000000062 * cos( 4.46845874247 +     6908.74604230380*t);
+   emb_x_0+=     0.00000000062 * cos( 2.18007483018 +     5657.40565767900*t);
+   emb_x_0+=     0.00000000050 * cos( 1.53898120800 +     2575.23909432920*t);
+   emb_x_0+=     0.00000000052 * cos( 1.82842431930 +     5244.04923920100*t);
+   emb_x_0+=     0.00000000050 * cos( 4.69455799921 +     5881.40372823420*t);
+   emb_x_0+=     0.00000000050 * cos( 1.95397557344 +     6684.74797174860*t);
+   emb_x_0+=     0.00000000053 * cos( 5.33677558065 +    12146.66705610760*t);
+   emb_x_0+=     0.00000000068 * cos( 6.19033473088 +    96677.89886304220*t);
+   emb_x_0+=     0.00000000064 * cos( 5.16841754323 +    24279.10701821359*t);
+   emb_x_0+=     0.00000000050 * cos( 3.21675900656 +    43738.80234596580*t);
+   emb_x_0+=     0.00000000050 * cos( 1.95435790721 +    23141.55838292460*t);
+   emb_x_0+=     0.00000000059 * cos( 4.28456134455 +     3981.49003408200*t);
+   emb_x_0+=     0.00000000058 * cos( 2.65359377943 +    40077.61957352000*t);
+   emb_x_0+=     0.00000000054 * cos( 3.56564135998 +    15664.03552270859*t);
+   emb_x_0+=     0.00000000061 * cos( 2.88962487651 +    40398.18991926600*t);
+   emb_x_0+=     0.00000000051 * cos( 3.00540767523 +     7668.97112632760*t);
+   emb_x_0+=     0.00000000048 * cos( 4.80656273565 +     3646.35037735440*t);
+   emb_x_0+=     0.00000000048 * cos( 2.80573655186 +      170.67287061920*t);
+   emb_x_0+=     0.00000000065 * cos( 6.08902325889 +    47162.51635463520*t);
+   emb_x_0+=     0.00000000056 * cos( 4.19093499866 +      799.82112516540*t);
+   emb_x_0+=     0.00000000050 * cos( 5.17709552474 +    17583.66007134780*t);
+   emb_x_0+=     0.00000000049 * cos( 5.16566731300 +     5490.30096152400*t);
+   emb_x_0+=     0.00000000058 * cos( 5.40778008235 +     6076.37511269480*t);
+   emb_x_0+=     0.00000000058 * cos( 1.24075349030 +     6489.77658728800*t);
+   emb_x_0+=     0.00000000059 * cos( 1.72439508836 +    44034.12758953940*t);
+   emb_x_0+=     0.00000000045 * cos( 1.62661962717 +     4480.70385926960*t);
+   emb_x_0+=     0.00000000061 * cos( 5.10264428859 +      491.66329245880*t);
+   emb_x_0+=     0.00000000044 * cos( 1.10033696885 +     2301.58581590939*t);
+   emb_x_0+=     0.00000000054 * cos( 2.79730041476 +     6006.33007812700*t);
+   emb_x_0+=     0.00000000054 * cos( 3.85123315789 +     6559.82162185580*t);
+   emb_x_0+=     0.00000000048 * cos( 1.42585982742 +    36147.40987730040*t);
+   emb_x_0+=     0.00000000044 * cos( 0.81389810430 +     5113.48759858280*t);
+   emb_x_0+=     0.00000000045 * cos( 4.22954077900 +      533.21408344360*t);
+   emb_x_0+=     0.00000000046 * cos( 2.27703250571 +     4384.72463205180*t);
+   emb_x_0+=     0.00000000048 * cos( 6.20173890368 +    35579.69123957000*t);
+   emb_x_0+=     0.00000000051 * cos( 0.20646603354 +    18326.65013188039*t);
+   emb_x_0+=     0.00000000043 * cos( 3.18939503640 +    27832.03821928320*t);
+   emb_x_0+=     0.00000000046 * cos( 4.72419541485 +    16703.06213349900*t);
+   emb_x_0+=     0.00000000052 * cos( 4.17179726389 +    11919.14086666800*t);
+   emb_x_0+=     0.00000000052 * cos( 0.71926561105 +    46386.99374331120*t);
+   emb_x_0+=     0.00000000047 * cos( 0.82248881139 +     5209.46682580060*t);
+   emb_x_0+=     0.00000000046 * cos( 4.04752097938 +    12555.10599971890*t);
+   emb_x_0+=     0.00000000052 * cos( 6.21526623133 +     5992.10298412540*t);
+   emb_x_0+=     0.00000000052 * cos( 0.43326734132 +     6574.04871585740*t);
+   emb_x_0+=     0.00000000044 * cos( 5.37918791865 +    69941.95360082899*t);
+   emb_x_0+=     0.00000000058 * cos( 0.36325356715 +    74673.98422817240*t);
+   emb_x_0+=     0.00000000056 * cos( 1.50089021543 +    27707.54249429480*t);
+   emb_x_0+=     0.00000000045 * cos( 6.22182195617 +    46360.69542351140*t);
+   emb_x_0+=     0.00000000043 * cos( 3.11608101277 +     8401.83971036980*t);
+   emb_x_0+=     0.00000000043 * cos( 1.13301289185 +    30639.85663863300*t);
+   emb_x_0+=     0.00000000049 * cos( 1.97844011728 +    22030.21295466960*t);
+   emb_x_0+=     0.00000000045 * cos( 4.79114307475 +    71519.29714327680*t);
+   emb_x_0+=     0.00000000051 * cos( 4.52703512787 +     7761.94242405580*t);
+   emb_x_0+=     0.00000000041 * cos( 0.43156447693 +     3735.23831175900*t);
+   emb_x_0+=     0.00000000041 * cos( 1.66415755313 +    44136.95134937400*t);
+   emb_x_0+=     0.00000000047 * cos( 4.45924253669 +     4708.23004870920*t);
+   emb_x_0+=     0.00000000056 * cos( 0.65498687965 +    28663.83165026540*t);
+   emb_x_0+=     0.00000000055 * cos( 3.26682599355 +    39302.09696219600*t);
+   emb_x_0+=     0.00000000049 * cos( 5.40929899528 +    27511.46787353720*t);
+   emb_x_0+=     0.00000000040 * cos( 0.05000438602 +    18208.34994259200*t);
+   emb_x_0+=     0.00000000053 * cos( 3.02524199997 +    38650.17350619900*t);
+   emb_x_0+=     0.00000000048 * cos( 2.50261741205 +    12662.13092720060*t);
+   emb_x_0+=     0.00000000052 * cos( 4.41165168727 +    12571.67462429020*t);
+   emb_x_0+=     0.00000000043 * cos( 4.69793306978 +     7322.10246078180*t);
+   emb_x_0+=     0.00000000041 * cos( 3.56541409677 +      664.75604513000*t);
+   emb_x_0+=     0.00000000044 * cos( 5.20087497624 +     6140.89722295520*t);
+   emb_x_0+=     0.00000000044 * cos( 1.44765859641 +     6425.25447702760*t);
+   emb_x_0+=     0.00000000039 * cos( 0.38207420659 +    22594.05489571199*t);
+   emb_x_0+=     0.00000000039 * cos( 4.22033276864 +    27278.46881644080*t);
+   emb_x_0+=     0.00000000045 * cos( 4.12311427457 +    44007.82926973960*t);
+   emb_x_0+=     0.00000000054 * cos( 1.95689060462 +     6560.11084373280*t);
+   emb_x_0+=     0.00000000054 * cos( 4.69164296803 +     6006.04085625000*t);
+   emb_x_0+=     0.00000000045 * cos( 5.54498349441 +   142861.39123625819*t);
+   emb_x_0+=     0.00000000047 * cos( 1.12159858858 +    24499.51966065240*t);
+   emb_x_0+=     0.00000000043 * cos( 1.36615498859 +     7026.06591052400*t);
+   emb_x_0+=     0.00000000040 * cos( 0.90738693747 +    26880.31981303260*t);
+   emb_x_0+=     0.00000000047 * cos( 4.50830744099 +    24485.29256665079*t);
+   emb_x_0+=     0.00000000043 * cos( 5.81439105801 +     6359.34192126700*t);
+   emb_x_0+=     0.00000000042 * cos( 4.70083045322 +     6231.79551212899*t);
+   emb_x_0+=     0.00000000042 * cos( 1.94770311943 +     6334.35618785381*t);
+   emb_x_0+=     0.00000000037 * cos( 5.39982800110 +    23020.65308658799*t);
+   emb_x_0+=     0.00000000052 * cos( 3.53311914292 +    11904.91877320180*t);
+   emb_x_0+=     0.00000000040 * cos( 1.50947925293 +     7075.85073845880*t);
+   emb_x_0+=     0.00000000049 * cos( 3.73016699917 +    32243.49850088060*t);
+   emb_x_0+=     0.00000000043 * cos( 4.45265815492 +    40000.04091585779*t);
+   emb_x_0+=     0.00000000040 * cos( 0.07156369559 +     7122.04513774180*t);
+   emb_x_0+=     0.00000000040 * cos( 0.29378456988 +     5444.10656224100*t);
+   emb_x_0+=     0.00000000038 * cos( 2.24527247949 +    44933.24935619040*t);
+   emb_x_0+=     0.00000000037 * cos( 0.09116813653 +     6850.79448772180*t);
+   emb_x_0+=     0.00000000037 * cos( 0.27418012894 +     5715.35721226100*t);
+   emb_x_0+=     0.00000000039 * cos( 2.88920508037 +       20.35531939880*t);
+   emb_x_0+=     0.00000000051 * cos( 5.82785362503 +    84111.74716305939*t);
+   emb_x_0+=     0.00000000037 * cos( 1.90653942815 +    16723.35014259500*t);
+   emb_x_0+=     0.00000000042 * cos( 3.79031051907 +     6328.21706962800*t);
+   emb_x_0+=     0.00000000042 * cos( 2.85822305358 +     6237.93463035480*t);
+   emb_x_0+=     0.00000000036 * cos( 2.35204592553 +    20760.42703319140*t);
+   emb_x_0+=     0.00000000049 * cos( 4.54040658078 +    11293.47067435560*t);
+   emb_x_0+=     0.00000000040 * cos( 3.64328287226 +     3215.13638064320*t);
+   emb_x_0+=     0.00000000040 * cos( 0.78396796338 +    71980.63357473118*t);
+
+   double emb_x_1=0.0;
+   emb_x_1+=     0.00123403046 * cos( 0.00000000000 +        0.00000000000*t);
+   emb_x_1+=     0.00051500156 * cos( 6.00266267204 +    12566.15169998280*t);
+   emb_x_1+=     0.00001290726 * cos( 5.95943124583 +    18849.22754997420*t);
+   emb_x_1+=     0.00001068627 * cos( 2.01554176551 +     6283.07584999140*t);
+   emb_x_1+=     0.00000212689 * cos( 1.73380190491 +     6279.55273164240*t);
+   emb_x_1+=     0.00000212515 * cos( 4.91489371033 +     6286.59896834040*t);
+   emb_x_1+=     0.00000062260 * cos( 0.36239798178 +     4705.73230754360*t);
+   emb_x_1+=     0.00000059822 * cos( 3.81195369871 +     6256.77753019160*t);
+   emb_x_1+=     0.00000059514 * cos( 2.83634160150 +     6309.37416979120*t);
+   emb_x_1+=     0.00000048841 * cos( 5.21419389335 +      775.52261132400*t);
+   emb_x_1+=     0.00000042883 * cos( 0.43789776559 +     1059.38193018920*t);
+   emb_x_1+=     0.00000046286 * cos( 0.01839494103 +     7860.41939243920*t);
+   emb_x_1+=     0.00000035675 * cos( 1.45279327264 +     5884.92684658320*t);
+   emb_x_1+=     0.00000036061 * cos( 2.16002201071 +     5753.38488489680*t);
+   emb_x_1+=     0.00000035367 * cos( 4.47243820095 +     6812.76681508600*t);
+   emb_x_1+=     0.00000032137 * cos( 5.19589851893 +     6681.22485339960*t);
+   emb_x_1+=     0.00000028763 * cos( 5.91618989512 +    25132.30339996560*t);
+   emb_x_1+=     0.00000028447 * cos( 1.14976253807 +     6127.65545055720*t);
+   emb_x_1+=     0.00000027573 * cos( 5.50119104683 +     6438.49624942560*t);
+   emb_x_1+=     0.00000024815 * cos( 2.92204909812 +     5486.77784317500*t);
+   emb_x_1+=     0.00000020611 * cos( 3.71790880968 +     7079.37385680780*t);
+   emb_x_1+=     0.00000019565 * cos( 2.89351924469 +     5507.55323866740*t);
+   emb_x_1+=     0.00000018308 * cos( 1.46954314992 +    11790.62908865880*t);
+   emb_x_1+=     0.00000016471 * cos( 6.22682639292 +    11506.76976979360*t);
+   emb_x_1+=     0.00000016757 * cos( 3.81935015812 +     7058.59846131540*t);
+   emb_x_1+=     0.00000014555 * cos( 5.97554823531 +     6290.18939699220*t);
+   emb_x_1+=     0.00000014388 * cos( 0.68157599309 +     6275.96230299060*t);
+   emb_x_1+=     0.00000013921 * cos( 1.44156172409 +      796.29800681640*t);
+   emb_x_1+=     0.00000011845 * cos( 4.15246503623 +     4694.00295470760*t);
+   emb_x_1+=     0.00000012573 * cos( 0.30429978572 +        7.11354700080*t);
+   emb_x_1+=     0.00000010073 * cos( 3.28940171828 +     3738.76143010800*t);
+   emb_x_1+=     0.00000009513 * cos( 4.62657118532 +     6069.77675455340*t);
+   emb_x_1+=     0.00000010425 * cos( 4.26610810148 +     6282.09552892320*t);
+   emb_x_1+=     0.00000010425 * cos( 2.38242547117 +     6284.05617105960*t);
+   emb_x_1+=     0.00000009565 * cos( 1.89315044337 +     4136.91043351620*t);
+   emb_x_1+=     0.00000009044 * cos( 2.00977420682 +     6496.37494542940*t);
+   emb_x_1+=     0.00000008543 * cos( 0.05563879210 +     1194.44701022460*t);
+   emb_x_1+=     0.00000008230 * cos( 1.25204318279 +     1589.07289528380*t);
+   emb_x_1+=     0.00000006181 * cos( 3.36235656603 +     8827.39026987480*t);
+   emb_x_1+=     0.00000006252 * cos( 4.75560512125 +     8429.24126646660*t);
+   emb_x_1+=     0.00000005525 * cos( 5.36960259598 +     4933.20844033260*t);
+   emb_x_1+=     0.00000007312 * cos( 0.49039562572 +     4535.05943692440*t);
+   emb_x_1+=     0.00000007074 * cos( 4.39447155183 +     5088.62883976680*t);
+   emb_x_1+=     0.00000005168 * cos( 5.15700586490 +    11769.85369316640*t);
+   emb_x_1+=     0.00000004878 * cos( 0.63279675325 +     6040.34724601740*t);
+   emb_x_1+=     0.00000005339 * cos( 5.63168593753 +     3154.68708489560*t);
+   emb_x_1+=     0.00000004985 * cos( 4.62777616911 +    12569.67481833180*t);
+   emb_x_1+=     0.00000004902 * cos( 4.14599178687 +     3930.20969621960*t);
+   emb_x_1+=     0.00000004709 * cos( 3.96401482904 +     5331.35744374080*t);
+   emb_x_1+=     0.00000004601 * cos( 6.02376082863 +     6525.80445396540*t);
+   emb_x_1+=     0.00000004331 * cos( 1.32999081524 +     7632.94325965020*t);
+   emb_x_1+=     0.00000005607 * cos( 2.55800588736 +     5729.50644714900*t);
+   emb_x_1+=     0.00000003968 * cos( 2.72696169781 +     7234.79425624200*t);
+   emb_x_1+=     0.00000004394 * cos( 0.04709283876 +    12559.03815298200*t);
+   emb_x_1+=     0.00000005262 * cos( 6.19471325132 +     8031.09226305840*t);
+   emb_x_1+=     0.00000005100 * cos( 4.10862449537 +     6836.64525283380*t);
+   emb_x_1+=     0.00000005075 * cos( 2.30582917471 +     7477.52286021600*t);
+   emb_x_1+=     0.00000003287 * cos( 5.41013501104 +    10977.07880469900*t);
+   emb_x_1+=     0.00000003366 * cos( 0.28947921379 +    11371.70468975820*t);
+   emb_x_1+=     0.00000003341 * cos( 1.64988056630 +     5223.69391980220*t);
+   emb_x_1+=     0.00000003495 * cos( 3.73571049570 +     4164.31198961300*t);
+   emb_x_1+=     0.00000004294 * cos( 4.93755825416 +     1592.59601363280*t);
+   emb_x_1+=     0.00000003408 * cos( 2.40593164878 +     3128.38876509580*t);
+   emb_x_1+=     0.00000002997 * cos( 4.86533819655 +     1748.01641306700*t);
+   emb_x_1+=     0.00000002750 * cos( 5.24429623672 +     7342.45778018060*t);
+   emb_x_1+=     0.00000002669 * cos( 3.42836090012 +     2146.16541647520*t);
+   emb_x_1+=     0.00000003313 * cos( 1.92016526057 +      801.82093112380*t);
+   emb_x_1+=     0.00000002991 * cos( 2.26433580653 +     8635.94200376320*t);
+   emb_x_1+=     0.00000003390 * cos( 4.23836515895 +      155.42039943420*t);
+   emb_x_1+=     0.00000002384 * cos( 5.74836325186 +     7084.89678111520*t);
+   emb_x_1+=     0.00000002288 * cos( 4.51088946570 +    17260.15465469040*t);
+   emb_x_1+=     0.00000002405 * cos( 5.96195251565 +     4690.47983635860*t);
+   emb_x_1+=     0.00000002052 * cos( 1.42012676041 +      951.71840625060*t);
+   emb_x_1+=     0.00000002049 * cos( 3.53634954129 +     1990.74501704100*t);
+   emb_x_1+=     0.00000002778 * cos( 1.46151348750 +     5481.25491886760*t);
+   emb_x_1+=     0.00000002656 * cos( 5.71405746008 +       26.29831979980*t);
+   emb_x_1+=     0.00000002135 * cos( 1.89020345470 +      536.80451209540*t);
+   emb_x_1+=     0.00000002529 * cos( 0.00975556739 +     1349.86740965880*t);
+   emb_x_1+=     0.00000001877 * cos( 2.82846563483 +      553.56940284240*t);
+   emb_x_1+=     0.00000001858 * cos( 0.74657310664 +      398.14900340820*t);
+   emb_x_1+=     0.00000002013 * cos( 2.40172210263 +      522.57741809380*t);
+   emb_x_1+=     0.00000001904 * cos( 4.23263119723 +     9437.76293488700*t);
+   emb_x_1+=     0.00000001708 * cos( 0.29065876129 +    10988.80815753500*t);
+   emb_x_1+=     0.00000001853 * cos( 1.84012190435 +     5230.80746680300*t);
+   emb_x_1+=     0.00000001857 * cos( 2.34009177634 +     5216.58037280140*t);
+   emb_x_1+=     0.00000001558 * cos( 1.10024744060 +      426.59819087600*t);
+   emb_x_1+=     0.00000001543 * cos( 5.14512865944 +    10447.38783960440*t);
+   emb_x_1+=     0.00000001555 * cos( 3.67004205114 +    13095.84266507740*t);
+   emb_x_1+=     0.00000001779 * cos( 1.80817581905 +    12562.62858163380*t);
+   emb_x_1+=     0.00000001450 * cos( 0.99962470222 +     6262.72053059260*t);
+   emb_x_1+=     0.00000001450 * cos( 5.64890887043 +     6303.43116939020*t);
+   emb_x_1+=     0.00000001880 * cos( 1.74872544631 +    10973.55568635000*t);
+   emb_x_1+=     0.00000001372 * cos( 0.89279228106 +     7875.67186362420*t);
+   emb_x_1+=     0.00000001247 * cos( 1.86318425701 +     2118.76386037840*t);
+   emb_x_1+=     0.00000001231 * cos( 4.10277044176 +    18319.53658487960*t);
+   emb_x_1+=     0.00000001373 * cos( 4.83873686423 +    12721.57209941700*t);
+   emb_x_1+=     0.00000001140 * cos( 3.49255962548 +     1066.49547719000*t);
+   emb_x_1+=     0.00000001127 * cos( 6.04843564608 +     5643.17856367740*t);
+   emb_x_1+=     0.00000001136 * cos( 2.83389341769 +     2942.46342329160*t);
+   emb_x_1+=     0.00000001071 * cos( 3.77511422440 +      242.72860397400*t);
+   emb_x_1+=     0.00000001165 * cos( 5.90779831262 +    15720.83878487840*t);
+   emb_x_1+=     0.00000000974 * cos( 0.24120093775 +      639.89728631400*t);
+   emb_x_1+=     0.00000000973 * cos( 1.30526267473 +     4292.33083295040*t);
+   emb_x_1+=     0.00000000993 * cos( 3.20822379065 +    10575.40668294180*t);
+   emb_x_1+=     0.00000000964 * cos( 2.13521649390 +     2388.89402044920*t);
+   emb_x_1+=     0.00000000929 * cos( 1.06241249364 +    16496.36139620240*t);
+   emb_x_1+=     0.00000000939 * cos( 5.78268735385 +    20426.57109242200*t);
+   emb_x_1+=     0.00000000889 * cos( 0.06581239696 +     6288.59877429880*t);
+   emb_x_1+=     0.00000000889 * cos( 0.29953586851 +     6277.55292568400*t);
+   emb_x_1+=     0.00000001018 * cos( 4.23379398746 +    12539.85338018300*t);
+   emb_x_1+=     0.00000001042 * cos( 2.30738424939 +      529.69096509460*t);
+   emb_x_1+=     0.00000001033 * cos( 1.04793549544 +     5760.49843189760*t);
+   emb_x_1+=     0.00000001017 * cos( 5.59172649862 +     6805.65326808520*t);
+   emb_x_1+=     0.00000000801 * cos( 5.36423767778 +    14314.16811304980*t);
+   emb_x_1+=     0.00000000795 * cos( 4.97031461103 +     6279.78949257360*t);
+   emb_x_1+=     0.00000000795 * cos( 1.67821896162 +     6286.36220740920*t);
+   emb_x_1+=     0.00000000723 * cos( 3.87714868475 +     9623.68827669120*t);
+   emb_x_1+=     0.00000000800 * cos( 4.75445675725 +    12029.34718788740*t);
+   emb_x_1+=     0.00000000698 * cos( 0.23124375738 +     6148.01076995600*t);
+   emb_x_1+=     0.00000000675 * cos( 5.56136235734 +    12964.30070339100*t);
+   emb_x_1+=     0.00000000678 * cos( 0.13408273301 +     6418.14093002680*t);
+   emb_x_1+=     0.00000000671 * cos( 1.28940504615 +     6709.67404086740*t);
+   emb_x_1+=     0.00000000898 * cos( 2.26352977340 +    13367.97263110660*t);
+   emb_x_1+=     0.00000000716 * cos( 4.07157593048 +    16730.46368959580*t);
+   emb_x_1+=     0.00000000651 * cos( 5.31876664276 +     5856.47765911540*t);
+   emb_x_1+=     0.00000000744 * cos( 4.25820617142 +    12043.57428188900*t);
+   emb_x_1+=     0.00000000650 * cos( 3.96082074858 +    14712.31711645800*t);
+   emb_x_1+=     0.00000000643 * cos( 6.14575788122 +     2352.86615377180*t);
+   emb_x_1+=     0.00000000666 * cos( 6.18664331112 +     5746.27133789600*t);
+   emb_x_1+=     0.00000000588 * cos( 3.24442116878 +    13119.72110282519*t);
+   emb_x_1+=     0.00000000627 * cos( 1.85323853183 +     4686.88940770680*t);
+   emb_x_1+=     0.00000000584 * cos( 1.78420561597 +    13517.87010623340*t);
+   emb_x_1+=     0.00000000610 * cos( 0.48716946879 +     6819.88036208680*t);
+   emb_x_1+=     0.00000000666 * cos( 3.29946942274 +     3634.62102451840*t);
+   emb_x_1+=     0.00000000585 * cos( 0.14860338726 +    12036.46073488820*t);
+   emb_x_1+=     0.00000000620 * cos( 5.87390167695 +    31415.37924995700*t);
+   emb_x_1+=     0.00000000553 * cos( 2.36074148933 +     6306.95428773920*t);
+   emb_x_1+=     0.00000000553 * cos( 4.28779208332 +     6259.19741224360*t);
+   emb_x_1+=     0.00000000527 * cos( 4.65845256725 +    10177.25767953360*t);
+   emb_x_1+=     0.00000000523 * cos( 5.55991825992 +    12139.55350910680*t);
+   emb_x_1+=     0.00000000514 * cos( 2.30343678959 +     4701.11650170840*t);
+   emb_x_1+=     0.00000000495 * cos( 3.58067041581 +    13362.44970679920*t);
+   emb_x_1+=     0.00000000688 * cos( 0.41160312810 +    13916.01910964160*t);
+   emb_x_1+=     0.00000000481 * cos( 4.14197696514 +     6286.95718534940*t);
+   emb_x_1+=     0.00000000481 * cos( 2.50655660751 +     6279.19451463340*t);
+   emb_x_1+=     0.00000000465 * cos( 0.81901657744 +    14143.49524243060*t);
+   emb_x_1+=     0.00000000440 * cos( 2.91975126929 +     3894.18182954220*t);
+   emb_x_1+=     0.00000000431 * cos( 3.09968854047 +    11499.65622279280*t);
+   emb_x_1+=     0.00000000452 * cos( 5.55409870988 +    12573.26524698360*t);
+   emb_x_1+=     0.00000000499 * cos( 1.08746978287 +    12410.73130054860*t);
+   emb_x_1+=     0.00000000440 * cos( 2.14902528642 +    24072.92146977640*t);
+   emb_x_1+=     0.00000000435 * cos( 4.76945528790 +     9917.69687450980*t);
+   emb_x_1+=     0.00000000398 * cos( 5.77181168519 +     8273.82086703240*t);
+   emb_x_1+=     0.00000000379 * cos( 4.62464502982 +     6386.16862421000*t);
+   emb_x_1+=     0.00000000379 * cos( 2.02388854283 +     6179.98307577280*t);
+   emb_x_1+=     0.00000000377 * cos( 5.90355422208 +     7872.14874527520*t);
+   emb_x_1+=     0.00000000419 * cos( 5.62219999256 +     5327.47610838280*t);
+   emb_x_1+=     0.00000000383 * cos( 0.40420536338 +     6319.10371666880*t);
+   emb_x_1+=     0.00000000383 * cos( 6.24432820927 +     6247.04798331399*t);
+   emb_x_1+=     0.00000000486 * cos( 3.05618432546 +    12592.45001978260*t);
+   emb_x_1+=     0.00000000457 * cos( 4.19200440747 +     6076.89030155420*t);
+   emb_x_1+=     0.00000000457 * cos( 2.45652916518 +     6489.26139842860*t);
+   emb_x_1+=     0.00000000512 * cos( 1.88365818219 +    11015.10647733480*t);
+   emb_x_1+=     0.00000000344 * cos( 0.28667648401 +    18451.07854656599*t);
+   emb_x_1+=     0.00000000340 * cos( 2.75068444234 +    12323.42309600880*t);
+   emb_x_1+=     0.00000000454 * cos( 0.72878821246 +     2787.04302385740*t);
+   emb_x_1+=     0.00000000364 * cos( 2.25070711701 +     6321.10352262720*t);
+   emb_x_1+=     0.00000000364 * cos( 4.39782645564 +     6245.04817735560*t);
+   emb_x_1+=     0.00000000342 * cos( 6.10019640486 +     6303.85124548380*t);
+   emb_x_1+=     0.00000000342 * cos( 0.54833716779 +     6262.30045449900*t);
+   emb_x_1+=     0.00000000316 * cos( 0.14807870151 +    11926.25441366880*t);
+   emb_x_1+=     0.00000000318 * cos( 1.32350341546 +     6311.52503745920*t);
+   emb_x_1+=     0.00000000318 * cos( 5.32503015719 +     6254.62666252360*t);
+   emb_x_1+=     0.00000000304 * cos( 2.00970548291 +    12779.45079542080*t);
+   emb_x_1+=     0.00000000360 * cos( 4.10116091510 +     1052.26838318840*t);
+   emb_x_1+=     0.00000000343 * cos( 2.93550529951 +     6267.82337880639*t);
+   emb_x_1+=     0.00000000343 * cos( 3.71302827314 +     6298.32832117640*t);
+   emb_x_1+=     0.00000000398 * cos( 0.12276624389 +     1551.04522264800*t);
+   emb_x_1+=     0.00000000367 * cos( 2.75102591636 +    10021.83728009940*t);
+   emb_x_1+=     0.00000000297 * cos( 0.51546903036 +     5429.87946823940*t);
+   emb_x_1+=     0.00000000337 * cos( 2.71244267795 +    17782.73207278420*t);
+   emb_x_1+=     0.00000000295 * cos( 4.05850212848 +     6702.56049386660*t);
+   emb_x_1+=     0.00000000331 * cos( 0.94287498279 +     7238.67559160000*t);
+   emb_x_1+=     0.00000000321 * cos( 2.23719619346 +    17796.95916678580*t);
+   emb_x_1+=     0.00000000378 * cos( 2.47077025009 +      213.29909543800*t);
+   emb_x_1+=     0.00000000268 * cos( 4.36573038796 +     2544.31441988340*t);
+   emb_x_1+=     0.00000000290 * cos( 5.31163321847 +      220.41264243880*t);
+   emb_x_1+=     0.00000000349 * cos( 0.95962542267 +    17789.84561978500*t);
+   emb_x_1+=     0.00000000275 * cos( 3.67150875811 +    14945.31617355440*t);
+   emb_x_1+=     0.00000000253 * cos( 2.22918080251 +    18073.70493865020*t);
+   emb_x_1+=     0.00000000257 * cos( 1.04749733742 +      955.59974160860*t);
+   emb_x_1+=     0.00000000234 * cos( 4.07399715957 +     1596.18644228460*t);
+   emb_x_1+=     0.00000000282 * cos( 4.50260494490 +    13341.67431130680*t);
+   emb_x_1+=     0.00000000280 * cos( 5.71871851487 +      206.18554843720*t);
+   emb_x_1+=     0.00000000269 * cos( 6.10957634874 +     9779.10867612540*t);
+   emb_x_1+=     0.00000000306 * cos( 3.20523259534 +     5863.59120611620*t);
+   emb_x_1+=     0.00000000216 * cos( 5.60561646872 +     3185.19202726560*t);
+   emb_x_1+=     0.00000000280 * cos( 1.91082976057 +     5642.19824260920*t);
+   emb_x_1+=     0.00000000233 * cos( 3.01339621844 +    10984.19235169980*t);
+   emb_x_1+=     0.00000000263 * cos( 2.16888500036 +    12168.00269657460*t);
+   emb_x_1+=     0.00000000226 * cos( 4.44534966925 +     3496.03282613400*t);
+   emb_x_1+=     0.00000000260 * cos( 3.16564185469 +    10969.96525769820*t);
+   emb_x_1+=     0.00000000227 * cos( 4.95666866168 +     5849.36411211460*t);
+   emb_x_1+=     0.00000000211 * cos( 1.22111379347 +    19651.04848109800*t);
+   emb_x_1+=     0.00000000250 * cos( 4.61595902292 +     6923.95345737360*t);
+   emb_x_1+=     0.00000000194 * cos( 2.37630770389 +      640.87760738220*t);
+   emb_x_1+=     0.00000000203 * cos( 4.39496828470 +     6062.66320755260*t);
+   emb_x_1+=     0.00000000203 * cos( 2.25356528795 +     6503.48849243020*t);
+   emb_x_1+=     0.00000000184 * cos( 0.13859812924 +    18209.33026366019*t);
+   emb_x_1+=     0.00000000184 * cos( 4.20955149479 +    12566.21901028560*t);
+   emb_x_1+=     0.00000000180 * cos( 5.76799248594 +      853.19638175200*t);
+   emb_x_1+=     0.00000000199 * cos( 2.84103239065 +        3.52311834900*t);
+   emb_x_1+=     0.00000000187 * cos( 4.59167688742 +     6276.02961329339*t);
+   emb_x_1+=     0.00000000187 * cos( 2.05685668523 +     6290.12208668940*t);
+   emb_x_1+=     0.00000000192 * cos( 2.99578290747 +     4732.03062734340*t);
+   emb_x_1+=     0.00000000184 * cos( 4.65365446464 +    12566.08438968000*t);
+   emb_x_1+=     0.00000000230 * cos( 2.23774195109 +    12567.13202105100*t);
+   emb_x_1+=     0.00000000165 * cos( 1.34307380910 +      110.20632121940*t);
+   emb_x_1+=     0.00000000166 * cos( 1.19477871058 +      433.71173787680*t);
+   emb_x_1+=     0.00000000184 * cos( 3.38049914956 +     6037.24420376200*t);
+   emb_x_1+=     0.00000000156 * cos( 2.34457936544 +     2648.45482547300*t);
+   emb_x_1+=     0.00000000175 * cos( 3.26771507312 +     6528.90749622080*t);
+   emb_x_1+=     0.00000000154 * cos( 0.97127667708 +     6112.40297937220*t);
+   emb_x_1+=     0.00000000154 * cos( 5.67725689556 +     6453.74872061060*t);
+   emb_x_1+=     0.00000000166 * cos( 4.51313560848 +    12352.85260454480*t);
+   emb_x_1+=     0.00000000146 * cos( 4.05664802756 +      149.56319713460*t);
+   emb_x_1+=     0.00000000172 * cos( 0.88734623011 +     1577.34354244780*t);
+   emb_x_1+=     0.00000000144 * cos( 3.76179690789 +    18635.92845453620*t);
+   emb_x_1+=     0.00000000146 * cos( 1.99164761433 +    13760.59871020740*t);
+   emb_x_1+=     0.00000000136 * cos( 0.75192359085 +    11712.95531823080*t);
+   emb_x_1+=     0.00000000142 * cos( 1.34533369944 +     4157.19844261220*t);
+   emb_x_1+=     0.00000000124 * cos( 0.99880484927 +    17256.63153634140*t);
+   emb_x_1+=     0.00000000140 * cos( 1.25451695814 +     9380.95967271720*t);
+   emb_x_1+=     0.00000000123 * cos( 0.26259126301 +     6172.86952877200*t);
+   emb_x_1+=     0.00000000123 * cos( 0.10275700246 +     6393.28217121080*t);
+   emb_x_1+=     0.00000000121 * cos( 5.40082218265 +       20.77539549240*t);
+   emb_x_1+=     0.00000000117 * cos( 4.47711713051 +    18852.75066832320*t);
+   emb_x_1+=     0.00000000120 * cos( 2.29184279740 +    17253.04110768959*t);
+   emb_x_1+=     0.00000000132 * cos( 4.94196827917 +    18875.52586977400*t);
+   emb_x_1+=     0.00000000123 * cos( 4.42839393502 +    12565.17137891460*t);
+   emb_x_1+=     0.00000000129 * cos( 2.56187150433 +    11513.88331679440*t);
+   emb_x_1+=     0.00000000121 * cos( 2.27781803851 +    10419.98628350760*t);
+   emb_x_1+=     0.00000000112 * cos( 0.58654969541 +      632.78373931320*t);
+   emb_x_1+=     0.00000000118 * cos( 5.96799701373 +    17654.78053974960*t);
+   emb_x_1+=     0.00000000105 * cos( 5.26799267002 +    23013.53953958720*t);
+   emb_x_1+=     0.00000000103 * cos( 0.24702432231 +     6271.34649715540*t);
+   emb_x_1+=     0.00000000103 * cos( 0.11832394316 +     6294.80520282740*t);
+   emb_x_1+=     0.00000000102 * cos( 2.69803619583 +    17267.26820169119*t);
+   emb_x_1+=     0.00000000108 * cos( 2.49639211639 +    23543.23050468179*t);
+   emb_x_1+=     0.00000000114 * cos( 2.21695343681 +     4171.42553661380*t);
+   emb_x_1+=     0.00000000101 * cos( 4.42100430236 +     4804.20927592700*t);
+   emb_x_1+=     0.00000000111 * cos( 3.66256237589 +     1581.95934828300*t);
+   emb_x_1+=     0.00000000119 * cos( 0.24456263546 +    10213.28554621100*t);
+   emb_x_1+=     0.00000000102 * cos( 5.67042027327 +     6549.68289171320*t);
+   emb_x_1+=     0.00000000102 * cos( 0.97811329938 +     6016.46880826960*t);
+   emb_x_1+=     0.00000000121 * cos( 0.05336972054 +    13625.53363017200*t);
+   emb_x_1+=     0.00000000090 * cos( 3.68828039857 +       76.26607127560*t);
+   emb_x_1+=     0.00000000086 * cos( 5.93949751968 +    12808.88030395680*t);
+   emb_x_1+=     0.00000000095 * cos( 1.34358559692 +    12345.73905754400*t);
+   emb_x_1+=     0.00000000106 * cos( 5.14281705612 +    11216.28429032400*t);
+   emb_x_1+=     0.00000000095 * cos( 2.50046439714 +    12012.58229714040*t);
+   emb_x_1+=     0.00000000081 * cos( 3.84772814821 +    11614.43329373220*t);
+   emb_x_1+=     0.00000000084 * cos( 1.14755262027 +     6255.67429389460*t);
+   emb_x_1+=     0.00000000084 * cos( 5.50098095238 +     6310.47740608820*t);
+   emb_x_1+=     0.00000000106 * cos( 6.00471264495 +     3097.88382272579*t);
+   emb_x_1+=     0.00000000080 * cos( 0.34711727404 +     6922.97313630540*t);
+   emb_x_1+=     0.00000000085 * cos( 5.79536877725 +     5636.06501667660*t);
+   emb_x_1+=     0.00000000092 * cos( 0.87789019103 +    10818.13528691580*t);
+   emb_x_1+=     0.00000000080 * cos( 4.19702650261 +     3583.34103067380*t);
+   emb_x_1+=     0.00000000083 * cos( 3.70902039500 +     5120.60114558360*t);
+   emb_x_1+=     0.00000000080 * cos( 1.41939312484 +    13521.75144159140*t);
+   emb_x_1+=     0.00000000084 * cos( 3.01891268791 +    10440.27429260360*t);
+   emb_x_1+=     0.00000000104 * cos( 5.42105253386 +     5650.29211067820*t);
+   emb_x_1+=     0.00000000086 * cos( 0.81273058913 +    12359.96615154560*t);
+   emb_x_1+=     0.00000000090 * cos( 6.02294397735 +     6195.76764545159*t);
+   emb_x_1+=     0.00000000090 * cos( 0.62558959530 +     6370.38405453120*t);
+   emb_x_1+=     0.00000000074 * cos( 2.33758314304 +     3340.61242669980*t);
+   emb_x_1+=     0.00000000082 * cos( 0.19767948100 +    18842.11400297339*t);
+   emb_x_1+=     0.00000000083 * cos( 5.31994916859 +    18415.51581209739*t);
+   emb_x_1+=     0.00000000101 * cos( 2.86191811970 +     3104.93005942380*t);
+   emb_x_1+=     0.00000000073 * cos( 6.09594424715 +      949.17560896980*t);
+   emb_x_1+=     0.00000000084 * cos( 4.35521501162 +     2379.16447357160*t);
+   emb_x_1+=     0.00000000075 * cos( 1.62833922577 +    30355.99731976780*t);
+   emb_x_1+=     0.00000000069 * cos( 3.24662354142 +    10454.50138660520*t);
+   emb_x_1+=     0.00000000076 * cos( 3.92665935350 +    17298.18232732620*t);
+   emb_x_1+=     0.00000000076 * cos( 4.83482959872 +    18429.74290609900*t);
+   emb_x_1+=     0.00000000079 * cos( 0.06205634713 +    22805.73556599360*t);
+   emb_x_1+=     0.00000000077 * cos( 2.05455640378 +    16858.48253293320*t);
+   emb_x_1+=     0.00000000087 * cos( 5.09755111164 +     3903.91137641980*t);
+   emb_x_1+=     0.00000000061 * cos( 3.70271756349 +    16737.57723659660*t);
+   emb_x_1+=     0.00000000078 * cos( 4.31638741861 +     8671.96987044060*t);
+   emb_x_1+=     0.00000000063 * cos( 0.99541849434 +    24356.78078864160*t);
+   emb_x_1+=     0.00000000075 * cos( 2.14136838637 +    23581.25817731760*t);
+   emb_x_1+=     0.00000000073 * cos( 4.49527510239 +     9388.00590941520*t);
+   emb_x_1+=     0.00000000058 * cos( 0.89406841198 +    17996.03116822220*t);
+   emb_x_1+=     0.00000000057 * cos( 2.48914084100 +      419.48464387520*t);
+   emb_x_1+=     0.00000000062 * cos( 4.06342805588 +    11925.27409260060*t);
+   emb_x_1+=     0.00000000060 * cos( 3.51548088474 +    34570.06633485260*t);
+   emb_x_1+=     0.00000000056 * cos( 5.33393809517 +    12455.94537876340*t);
+   emb_x_1+=     0.00000000056 * cos( 2.17965993818 +    11764.33076885900*t);
+   emb_x_1+=     0.00000000061 * cos( 0.15194461973 +    28237.23345938940*t);
+   emb_x_1+=     0.00000000065 * cos( 5.87045161229 +     7335.34423317980*t);
+   emb_x_1+=     0.00000000072 * cos( 5.87296425674 +     6510.55198278040*t);
+   emb_x_1+=     0.00000000072 * cos( 0.77556931591 +     6055.59971720240*t);
+   emb_x_1+=     0.00000000071 * cos( 0.23083621209 +    12545.37630449040*t);
+   emb_x_1+=     0.00000000061 * cos( 5.37705181537 +     6253.64634145540*t);
+   emb_x_1+=     0.00000000061 * cos( 1.27148175728 +     6312.50535852740*t);
+   emb_x_1+=     0.00000000055 * cos( 1.18603703570 +     2699.73481931760*t);
+   emb_x_1+=     0.00000000055 * cos( 5.94073604896 +    15671.08175940660*t);
+   emb_x_1+=     0.00000000052 * cos( 0.26563539971 +     5326.78669402080*t);
+   emb_x_1+=     0.00000000053 * cos( 1.77739578512 +    26735.94526221320*t);
+   emb_x_1+=     0.00000000065 * cos( 2.33938496802 +    14919.01785375460*t);
+   emb_x_1+=     0.00000000049 * cos( 4.18739781015 +    35579.69123957000*t);
+   emb_x_1+=     0.00000000053 * cos( 5.28549359722 +     6357.85744855870*t);
+   emb_x_1+=     0.00000000053 * cos( 1.36303997543 +     6208.29425142410*t);
+   emb_x_1+=     0.00000000051 * cos( 0.76407742618 +    16522.65971600220*t);
+   emb_x_1+=     0.00000000049 * cos( 4.92674479666 +    21393.54196985760*t);
+   emb_x_1+=     0.00000000049 * cos( 0.60525485888 +     7018.95236352320*t);
+   emb_x_1+=     0.00000000049 * cos( 6.04327871377 +     5547.19933645960*t);
+   emb_x_1+=     0.00000000049 * cos( 2.59123701567 +    12416.58850284820*t);
+   emb_x_1+=     0.00000000050 * cos( 0.07523239468 +    26084.02180621620*t);
+   emb_x_1+=     0.00000000048 * cos( 5.46985443988 +    12132.43996210600*t);
+   emb_x_1+=     0.00000000051 * cos( 3.40683294217 +    23006.42599258639*t);
+   emb_x_1+=     0.00000000058 * cos( 2.46543150418 +    15141.39079431200*t);
+   emb_x_1+=     0.00000000048 * cos( 4.88190521609 +    16200.77272450120*t);
+   emb_x_1+=     0.00000000047 * cos( 0.97667389279 +    24279.10701821359*t);
+   emb_x_1+=     0.00000000049 * cos( 1.66846784739 +     5749.86176654780*t);
+   emb_x_1+=     0.00000000049 * cos( 4.98006572526 +     6816.28993343500*t);
+   emb_x_1+=     0.00000000046 * cos( 3.45134537804 +    19378.91851506880*t);
+   emb_x_1+=     0.00000000047 * cos( 2.09300174029 +    24065.80792277559*t);
+   emb_x_1+=     0.00000000046 * cos( 5.73797450552 +   104351.61256629678*t);
+   emb_x_1+=     0.00000000049 * cos( 1.45701599108 +     8982.81066930900*t);
+   emb_x_1+=     0.00000000050 * cos( 4.49441848554 +    24602.61243487099*t);
+   emb_x_1+=     0.00000000045 * cos( 0.12799634054 +    24491.42579258340*t);
+   emb_x_1+=     0.00000000049 * cos( 2.95918505969 +    16460.33352952499*t);
+   emb_x_1+=     0.00000000054 * cos( 0.03974674424 +     7834.12107263940*t);
+   emb_x_1+=     0.00000000044 * cos( 3.92726741188 +     8662.24032356300*t);
+   emb_x_1+=     0.00000000049 * cos( 4.51486356794 +    29303.72893657940*t);
+   emb_x_1+=     0.00000000049 * cos( 4.98978743446 +    29289.50184257780*t);
+   emb_x_1+=     0.00000000043 * cos( 3.34677142688 +     5244.04923920100*t);
+   emb_x_1+=     0.00000000044 * cos( 3.63764737159 +     6175.41232605280*t);
+   emb_x_1+=     0.00000000044 * cos( 3.01088620106 +     6390.73937393000*t);
+   emb_x_1+=     0.00000000042 * cos( 1.47305743112 +    24705.70520908959*t);
+   emb_x_1+=     0.00000000041 * cos( 4.78696656856 +    17370.36097590980*t);
+   emb_x_1+=     0.00000000041 * cos( 4.99610869214 +    22483.84857449259*t);
+   emb_x_1+=     0.00000000041 * cos( 3.62603303245 +    23020.65308658799*t);
+   emb_x_1+=     0.00000000045 * cos( 2.13147646951 +    54247.41313575040*t);
+   emb_x_1+=     0.00000000040 * cos( 3.00393229632 +    32243.49850088060*t);
+   emb_x_1+=     0.00000000040 * cos( 0.45946989377 +    24336.00539314920*t);
+   emb_x_1+=     0.00000000040 * cos( 0.00254937606 +     8401.83971036980*t);
+   emb_x_1+=     0.00000000041 * cos( 5.05922003836 +    26482.17080962440*t);
+   emb_x_1+=     0.00000000044 * cos( 4.93011217974 +    28230.18722269139*t);
+   emb_x_1+=     0.00000000039 * cos( 5.54872052106 +    45585.17281218740*t);
+   emb_x_1+=     0.00000000038 * cos( 3.93561558747 +     6268.84875598980*t);
+   emb_x_1+=     0.00000000038 * cos( 2.71291798518 +     6297.30294399300*t);
+   emb_x_1+=     0.00000000041 * cos( 4.95311845784 +    15508.61512327440*t);
+   emb_x_1+=     0.00000000038 * cos( 0.10976793876 +     6171.88920770380*t);
+   emb_x_1+=     0.00000000038 * cos( 0.25558032671 +     6394.26249227900*t);
+   emb_x_1+=     0.00000000039 * cos( 3.31511780881 +    29826.30635467320*t);
+   emb_x_1+=     0.00000000038 * cos( 4.03952051907 +    16062.18452611680*t);
+   emb_x_1+=     0.00000000037 * cos( 2.07587316842 +    20452.86941222180*t);
+   emb_x_1+=     0.00000000049 * cos( 5.78169307611 +    18216.44381066100*t);
+   emb_x_1+=     0.00000000041 * cos( 1.35496527372 +    29088.81141598500*t);
+   emb_x_1+=     0.00000000038 * cos( 6.16268217428 +    18202.21671665939*t);
+   emb_x_1+=     0.00000000036 * cos( 3.84319801860 +    39302.09696219600*t);
+   emb_x_1+=     0.00000000038 * cos( 6.14920442251 +    26709.64694241340*t);
+   emb_x_1+=     0.00000000036 * cos( 6.28049708584 +     9924.81042151060*t);
+   emb_x_1+=     0.00000000035 * cos( 4.66732134451 +     4061.21921539440*t);
+   emb_x_1+=     0.00000000040 * cos( 3.43223548088 +    27511.46787353720*t);
+   emb_x_1+=     0.00000000034 * cos( 1.38801022558 +     6147.45052498140*t);
+   emb_x_1+=     0.00000000034 * cos( 5.26052334707 +     6418.70117500140*t);
+   emb_x_1+=     0.00000000034 * cos( 1.97967478962 +    17271.88400752640*t);
+   emb_x_1+=     0.00000000037 * cos( 6.05744777692 +    21150.81336588360*t);
+   emb_x_1+=     0.00000000034 * cos( 1.90509484306 +    29026.48522950779*t);
+   emb_x_1+=     0.00000000035 * cos( 3.20661690466 +     7322.10246078180*t);
+   emb_x_1+=     0.00000000044 * cos( 2.79017179925 +     3981.49003408200*t);
+   emb_x_1+=     0.00000000034 * cos( 1.21719184072 +    21548.96236929180*t);
+   emb_x_1+=     0.00000000033 * cos( 3.22867201822 +    30666.15495843280*t);
+   emb_x_1+=     0.00000000034 * cos( 4.95128499531 +    38500.27603107220*t);
+   emb_x_1+=     0.00000000032 * cos( 4.38302389502 +    17686.75284556640*t);
+   emb_x_1+=     0.00000000038 * cos( 4.95582614524 +    19004.64794940840*t);
+   emb_x_1+=     0.00000000033 * cos( 1.70921741352 +    34520.30930938080*t);
+   emb_x_1+=     0.00000000041 * cos( 4.09101902227 +    21202.09370374600*t);
+   emb_x_1+=     0.00000000031 * cos( 0.28544969417 +    22003.91463486980*t);
+   emb_x_1+=     0.00000000037 * cos( 1.21597687953 +     7349.57132718140*t);
+   emb_x_1+=     0.00000000031 * cos( 0.19462514993 +    27832.03821928320*t);
+   emb_x_1+=     0.00000000037 * cos( 3.41912843249 +    29296.61538957860*t);
+   emb_x_1+=     0.00000000035 * cos( 2.69261244483 +    11919.14086666800*t);
+   emb_x_1+=     0.00000000035 * cos( 0.83074174405 +    29062.51309618520*t);
+   emb_x_1+=     0.00000000030 * cos( 2.96012497927 +    12489.88562870720*t);
+   emb_x_1+=     0.00000000035 * cos( 4.83736876925 +    21228.39202354580*t);
+   emb_x_1+=     0.00000000032 * cos( 6.26164802166 +     6231.79551212899*t);
+   emb_x_1+=     0.00000000032 * cos( 0.38688555099 +     6334.35618785381*t);
+   emb_x_1+=     0.00000000029 * cos( 0.70432070922 +    49515.38250840700*t);
+   emb_x_1+=     0.00000000032 * cos( 0.44788919628 +    31172.65064598300*t);
+   emb_x_1+=     0.00000000031 * cos( 2.69970301246 +    33019.02111220460*t);
+   emb_x_1+=     0.00000000029 * cos( 2.26256470092 +    27278.46881644080*t);
+   emb_x_1+=     0.00000000031 * cos( 3.19023486538 +     5017.50837136500*t);
+   emb_x_1+=     0.00000000028 * cos( 1.87849143707 +    24080.03501677719*t);
+   emb_x_1+=     0.00000000034 * cos( 0.78574341965 +     5905.70224207560*t);
+   emb_x_1+=     0.00000000034 * cos( 5.86279015300 +     6660.44945790720*t);
+   emb_x_1+=     0.00000000038 * cos( 5.12224209756 +    31570.79964939120*t);
+   emb_x_1+=     0.00000000027 * cos( 3.76025315460 +    19402.79695281660*t);
+   emb_x_1+=     0.00000000030 * cos( 2.32405719366 +    12586.92709547520*t);
+   emb_x_1+=     0.00000000028 * cos( 5.54772804461 +    34911.41207609100*t);
+   emb_x_1+=     0.00000000027 * cos( 5.30079330746 +    43232.30665841560*t);
+   emb_x_1+=     0.00000000027 * cos( 6.17189140209 +     1692.16566950240*t);
+   emb_x_1+=     0.00000000036 * cos( 4.81217077776 +    15664.03552270859*t);
+   emb_x_1+=     0.00000000027 * cos( 0.65637418130 +     6360.74962041940*t);
+   emb_x_1+=     0.00000000027 * cos( 5.99215939135 +     6205.40207956340*t);
+   emb_x_1+=     0.00000000026 * cos( 3.42413737440 +    24383.07910844140*t);
+   emb_x_1+=     0.00000000033 * cos( 3.45722382528 +    28628.33622609960*t);
+   emb_x_1+=     0.00000000027 * cos( 2.32049228976 +    32367.09765620760*t);
+   emb_x_1+=     0.00000000026 * cos( 1.98182489066 +    19804.82729158280*t);
+   emb_x_1+=     0.00000000027 * cos( 1.77114813957 +     4487.81740627040*t);
+   emb_x_1+=     0.00000000028 * cos( 2.89944615107 +    58953.14544329400*t);
+   emb_x_1+=     0.00000000028 * cos( 3.80331477821 +    10239.58386601080*t);
+   emb_x_1+=     0.00000000026 * cos( 0.36912120822 +     6339.97422492700*t);
+   emb_x_1+=     0.00000000026 * cos( 6.27941236443 +     6226.17747505580*t);
+   emb_x_1+=     0.00000000026 * cos( 6.01416870867 +     4797.09572892620*t);
+   emb_x_1+=     0.00000000028 * cos( 2.80825729869 +    22476.73502749179*t);
+   emb_x_1+=     0.00000000029 * cos( 0.17746825116 +     6592.35417264720*t);
+   emb_x_1+=     0.00000000029 * cos( 0.18788001431 +     5973.79752733560*t);
+   emb_x_1+=     0.00000000026 * cos( 3.93576097012 +    36949.23080842420*t);
+   emb_x_1+=     0.00000000027 * cos( 0.64682460545 +    29424.63423291600*t);
+   emb_x_1+=     0.00000000026 * cos( 2.54059552040 +     9910.58332750900*t);
+   emb_x_1+=     0.00000000028 * cos( 2.41855003101 +    46386.99374331120*t);
+   emb_x_1+=     0.00000000024 * cos( 1.22629501602 +    21947.11137270000*t);
+   emb_x_1+=     0.00000000024 * cos( 0.55027097678 +    11616.97609101300*t);
+   emb_x_1+=     0.00000000025 * cos( 2.78491065355 +     6302.74561089119*t);
+   emb_x_1+=     0.00000000025 * cos( 3.86362291910 +     6263.40608909160*t);
+   emb_x_1+=     0.00000000026 * cos( 4.44347824401 +    15110.46611986620*t);
+   emb_x_1+=     0.00000000027 * cos( 5.91935985488 +    52175.80628314840*t);
+   emb_x_1+=     0.00000000028 * cos( 6.10918595816 +    18052.92954315780*t);
+   emb_x_1+=     0.00000000023 * cos( 1.99242379288 +    30774.50164257480*t);
+   emb_x_1+=     0.00000000024 * cos( 3.73458431258 +    31968.94865279940*t);
+   emb_x_1+=     0.00000000023 * cos( 0.62347300029 +    17892.93839400359*t);
+   emb_x_1+=     0.00000000024 * cos( 3.21260068222 +    23550.34405168259*t);
+   emb_x_1+=     0.00000000023 * cos( 1.36063534536 +     6716.78758786820*t);
+   emb_x_1+=     0.00000000027 * cos( 0.02724821878 +    25158.60171976540*t);
+   emb_x_1+=     0.00000000026 * cos( 3.57645528088 +     8584.66166590080*t);
+   emb_x_1+=     0.00000000022 * cos( 3.99741791962 +     7445.55055439920*t);
+   emb_x_1+=     0.00000000022 * cos( 6.10694895695 +    15265.88651930040*t);
+   emb_x_1+=     0.00000000025 * cos( 0.81555419576 +    14155.22459526660*t);
+   emb_x_1+=     0.00000000023 * cos( 3.06283945185 +     6272.03014972750*t);
+   emb_x_1+=     0.00000000023 * cos( 3.58569412080 +     6294.12155025530*t);
+   emb_x_1+=     0.00000000025 * cos( 2.44880076940 +    20199.09495963300*t);
+   emb_x_1+=     0.00000000021 * cos( 5.78253540504 +    22779.43724619380*t);
+   emb_x_1+=     0.00000000023 * cos( 2.25096306087 +    34115.11406927460*t);
+   emb_x_1+=     0.00000000022 * cos( 4.67990936559 +    34596.36465465240*t);
+   emb_x_1+=     0.00000000029 * cos( 0.93434607315 +    11933.36796066960*t);
+   emb_x_1+=     0.00000000021 * cos( 4.08577407328 +    30220.93223973240*t);
+   emb_x_1+=     0.00000000020 * cos( 0.47093224249 +    47162.51635463520*t);
+   emb_x_1+=     0.00000000022 * cos( 1.38262024424 +     4379.63903749020*t);
+   emb_x_1+=     0.00000000020 * cos( 5.31355969163 +    40879.44050464380*t);
+   emb_x_1+=     0.00000000022 * cos( 4.52226925108 +    28759.81087748319*t);
+   emb_x_1+=     0.00000000026 * cos( 3.41054638032 +    18422.62935909819*t);
+   emb_x_1+=     0.00000000027 * cos( 0.67002293804 +    34513.26307268280*t);
+   emb_x_1+=     0.00000000019 * cos( 4.78610064760 +    28313.28880466100*t);
+   emb_x_1+=     0.00000000020 * cos( 5.09152714894 +     8186.51266249260*t);
+   emb_x_1+=     0.00000000025 * cos( 3.53211055104 +    55798.45835839840*t);
+   emb_x_1+=     0.00000000025 * cos( 0.38024667868 +     9225.53927328300*t);
+   emb_x_1+=     0.00000000020 * cos( 1.71238700404 +    27433.88921587499*t);
+   emb_x_1+=     0.00000000024 * cos( 1.20856889899 +     9411.46461508720*t);
+   emb_x_1+=     0.00000000020 * cos( 1.65292451759 +     6379.05507720920*t);
+   emb_x_1+=     0.00000000020 * cos( 4.99560905506 +     6187.09662277360*t);
+   emb_x_1+=     0.00000000020 * cos( 2.94178503134 +    19800.94595622480*t);
+   emb_x_1+=     0.00000000020 * cos( 4.32349571107 +    62883.35513951360*t);
+   emb_x_1+=     0.00000000022 * cos( 4.44470422678 +    18326.65013188039*t);
+   emb_x_1+=     0.00000000019 * cos( 4.37840329985 +    16207.88627150200*t);
+   emb_x_1+=     0.00000000022 * cos( 4.46505143635 +     6332.83287546320*t);
+   emb_x_1+=     0.00000000022 * cos( 2.18348213630 +     6233.31882451960*t);
+   emb_x_1+=     0.00000000019 * cos( 1.36569875427 +    23539.70738633280*t);
+   emb_x_1+=     0.00000000018 * cos( 4.61817062443 +     6593.91664885980*t);
+   emb_x_1+=     0.00000000018 * cos( 2.03036294822 +     5972.23505112300*t);
+   emb_x_1+=     0.00000000019 * cos( 2.36455752682 +     5618.31980486140*t);
+   emb_x_1+=     0.00000000019 * cos( 4.28397604583 +     6947.83189512140*t);
+   emb_x_1+=     0.00000000020 * cos( 2.85390720284 +     9070.11887384880*t);
+   emb_x_1=emb_x_1 * t;
+
+   double emb_x_2=0.0;
+   emb_x_2+=     0.00004143217 * cos( 3.14159265359 +        0.00000000000*t);
+   emb_x_2+=     0.00002175695 * cos( 4.39999849572 +    12566.15169998280*t);
+   emb_x_2+=     0.00000995233 * cos( 0.20790847155 +     6283.07584999140*t);
+   emb_x_2+=     0.00000092659 * cos( 4.19285471010 +    18849.22754997420*t);
+   emb_x_2+=     0.00000013679 * cos( 3.37833642063 +     6286.59896834040*t);
+   emb_x_2+=     0.00000013668 * cos( 3.27271492019 +     6279.55273164240*t);
+   emb_x_2+=     0.00000008190 * cos( 4.41928022030 +     6256.77753019160*t);
+   emb_x_2+=     0.00000008145 * cos( 2.22925093408 +     6309.37416979120*t);
+   emb_x_2+=     0.00000008004 * cos( 5.92288197328 +     6127.65545055720*t);
+   emb_x_2+=     0.00000007746 * cos( 0.72469358509 +     6438.49624942560*t);
+   emb_x_2+=     0.00000003737 * cos( 2.48724420764 +     5507.55323866740*t);
+   emb_x_2+=     0.00000002893 * cos( 4.94972697460 +     4705.73230754360*t);
+   emb_x_2+=     0.00000002937 * cos( 4.08160813245 +    25132.30339996560*t);
+   emb_x_2+=     0.00000003203 * cos( 4.14275083744 +     7058.59846131540*t);
+   emb_x_2+=     0.00000002254 * cos( 3.51957213521 +      775.52261132400*t);
+   emb_x_2+=     0.00000002170 * cos( 1.70495739017 +     7860.41939243920*t);
+   emb_x_2+=     0.00000002031 * cos( 2.29363438970 +     6275.96230299060*t);
+   emb_x_2+=     0.00000002031 * cos( 4.35489918295 +     6290.18939699220*t);
+   emb_x_2+=     0.00000001904 * cos( 4.46668961834 +     5486.77784317500*t);
+   emb_x_2+=     0.00000001720 * cos( 5.23142094282 +     1059.38193018920*t);
+   emb_x_2+=     0.00000001548 * cos( 2.12906745749 +     6040.34724601740*t);
+   emb_x_2+=     0.00000001563 * cos( 2.18101487697 +     7079.37385680780*t);
+   emb_x_2+=     0.00000001456 * cos( 3.53770823599 +     5753.38488489680*t);
+   emb_x_2+=     0.00000001424 * cos( 3.16250984559 +     6812.76681508600*t);
+   emb_x_2+=     0.00000001470 * cos( 4.51659549398 +     6525.80445396540*t);
+   emb_x_2+=     0.00000001429 * cos( 3.01998118552 +     5884.92684658320*t);
+   emb_x_2+=     0.00000001288 * cos( 3.62892288697 +     6681.22485339960*t);
+   emb_x_2+=     0.00000001338 * cos( 1.04445705185 +     5729.50644714900*t);
+   emb_x_2+=     0.00000001039 * cos( 0.79543243402 +     6284.05617105960*t);
+   emb_x_2+=     0.00000001039 * cos( 5.85310113863 +     6282.09552892320*t);
+   emb_x_2+=     0.00000001212 * cos( 5.61713170172 +     6836.64525283380*t);
+   emb_x_2+=     0.00000000904 * cos( 2.43535199673 +     5331.35744374080*t);
+   emb_x_2+=     0.00000000871 * cos( 3.84816000443 +     4933.20844033260*t);
+   emb_x_2+=     0.00000000879 * cos( 3.06133930098 +    11790.62908865880*t);
+   emb_x_2+=     0.00000000769 * cos( 4.22812146942 +     7234.79425624200*t);
+   emb_x_2+=     0.00000000744 * cos( 0.34956564055 +     4136.91043351620*t);
+   emb_x_2+=     0.00000000675 * cos( 2.82587446878 +     7632.94325965020*t);
+   emb_x_2+=     0.00000000654 * cos( 1.40322441961 +    11506.76976979360*t);
+   emb_x_2+=     0.00000000854 * cos( 5.24328642976 +     4535.05943692440*t);
+   emb_x_2+=     0.00000000665 * cos( 4.81749499447 +     1194.44701022460*t);
+   emb_x_2+=     0.00000000834 * cos( 5.75118975782 +      155.42039943420*t);
+   emb_x_2+=     0.00000000785 * cos( 5.88776299473 +     5088.62883976680*t);
+   emb_x_2+=     0.00000000548 * cos( 6.12128965864 +      796.29800681640*t);
+   emb_x_2+=     0.00000000686 * cos( 1.90135823927 +      801.82093112380*t);
+   emb_x_2+=     0.00000000507 * cos( 5.42408029291 +     4694.00295470760*t);
+   emb_x_2+=     0.00000000677 * cos( 2.71965226391 +     3128.38876509580*t);
+   emb_x_2+=     0.00000000499 * cos( 3.02752235850 +    12569.67481833180*t);
+   emb_x_2+=     0.00000000480 * cos( 0.01327044252 +     8429.24126646660*t);
+   emb_x_2+=     0.00000000592 * cos( 1.39433962245 +     8031.09226305840*t);
+   emb_x_2+=     0.00000000498 * cos( 6.19000103499 +     1589.07289528380*t);
+   emb_x_2+=     0.00000000562 * cos( 0.77105088912 +     7477.52286021600*t);
+   emb_x_2+=     0.00000000415 * cos( 0.00558962450 +     6069.77675455340*t);
+   emb_x_2+=     0.00000000415 * cos( 0.35131226280 +     6496.37494542940*t);
+   emb_x_2+=     0.00000000402 * cos( 1.72177128440 +     3738.76143010800*t);
+   emb_x_2+=     0.00000000503 * cos( 3.41456622371 +     1592.59601363280*t);
+   emb_x_2+=     0.00000000390 * cos( 5.36711219721 +       26.29831979980*t);
+   emb_x_2+=     0.00000000372 * cos( 4.35153585802 +      553.56940284240*t);
+   emb_x_2+=     0.00000000325 * cos( 2.94948293011 +      951.71840625060*t);
+   emb_x_2+=     0.00000000321 * cos( 2.01798251637 +     1990.74501704100*t);
+   emb_x_2+=     0.00000000335 * cos( 1.12608331535 +     4690.47983635860*t);
+   emb_x_2+=     0.00000000366 * cos( 0.10449761483 +    12721.57209941700*t);
+   emb_x_2+=     0.00000000287 * cos( 2.26476426223 +      242.72860397400*t);
+   emb_x_2+=     0.00000000347 * cos( 1.34248948630 +     5481.25491886760*t);
+   emb_x_2+=     0.00000000368 * cos( 3.91910931412 +     9437.76293488700*t);
+   emb_x_2+=     0.00000000255 * cos( 1.83876117314 +    11371.70468975820*t);
+   emb_x_2+=     0.00000000242 * cos( 4.92327001310 +     8827.39026987480*t);
+   emb_x_2+=     0.00000000307 * cos( 1.56817330160 +     1349.86740965880*t);
+   emb_x_2+=     0.00000000242 * cos( 0.12263961387 +     1748.01641306700*t);
+   emb_x_2+=     0.00000000222 * cos( 4.88851960168 +     4164.31198961300*t);
+   emb_x_2+=     0.00000000230 * cos( 1.05870192923 +     3154.68708489560*t);
+   emb_x_2+=     0.00000000204 * cos( 0.44240387853 +    11769.85369316640*t);
+   emb_x_2+=     0.00000000188 * cos( 5.65888993470 +     7875.67186362420*t);
+   emb_x_2+=     0.00000000189 * cos( 0.61654841817 +     2388.89402044920*t);
+   emb_x_2+=     0.00000000233 * cos( 5.23023625373 +     7084.89678111520*t);
+   emb_x_2+=     0.00000000183 * cos( 4.10128236867 +     7342.45778018060*t);
+   emb_x_2+=     0.00000000173 * cos( 0.52511015365 +    10977.07880469900*t);
+   emb_x_2+=     0.00000000159 * cos( 2.72294432534 +     4292.33083295040*t);
+   emb_x_2+=     0.00000000207 * cos( 2.28587256917 +    13367.97263110660*t);
+   emb_x_2+=     0.00000000202 * cos( 3.21445201981 +    10973.55568635000*t);
+   emb_x_2+=     0.00000000145 * cos( 5.45882589245 +      398.14900340820*t);
+   emb_x_2+=     0.00000000155 * cos( 5.81412827494 +    12410.73130054860*t);
+   emb_x_2+=     0.00000000134 * cos( 4.56296979111 +    10575.40668294180*t);
+   emb_x_2+=     0.00000000127 * cos( 4.84612567560 +    13119.72110282519*t);
+   emb_x_2+=     0.00000000115 * cos( 4.96539448045 +     2146.16541647520*t);
+   emb_x_2+=     0.00000000125 * cos( 3.87814283336 +    12562.62858163380*t);
+   emb_x_2+=     0.00000000125 * cos( 0.17913722961 +     8635.94200376320*t);
+   emb_x_2+=     0.00000000122 * cos( 1.05019176110 +    12592.45001978260*t);
+   emb_x_2+=     0.00000000111 * cos( 2.15765835240 +     1577.34354244780*t);
+   emb_x_2+=     0.00000000134 * cos( 6.25201667258 +     3930.20969621960*t);
+   emb_x_2+=     0.00000000107 * cos( 3.43601884429 +    13517.87010623340*t);
+   emb_x_2+=     0.00000000102 * cos( 4.37125144267 +    10988.80815753500*t);
+   emb_x_2+=     0.00000000095 * cos( 5.98656941724 +    17260.15465469040*t);
+   emb_x_2+=     0.00000000100 * cos( 2.34969093002 +    13095.84266507740*t);
+   emb_x_2+=     0.00000000110 * cos( 3.63768958070 +     5223.69391980220*t);
+   emb_x_2+=     0.00000000085 * cos( 4.32512089859 +     3894.18182954220*t);
+   emb_x_2+=     0.00000000109 * cos( 3.99726753431 +    12539.85338018300*t);
+   emb_x_2+=     0.00000000079 * cos( 5.22586609526 +    16730.46368959580*t);
+   emb_x_2+=     0.00000000079 * cos( 0.08497701397 +    10447.38783960440*t);
+   emb_x_2+=     0.00000000076 * cos( 4.11846374330 +    12323.42309600880*t);
+   emb_x_2+=     0.00000000105 * cos( 3.53845510796 +     5642.19824260920*t);
+   emb_x_2+=     0.00000000075 * cos( 5.73067319226 +     6709.67404086740*t);
+   emb_x_2+=     0.00000000075 * cos( 0.91786038039 +     5856.47765911540*t);
+   emb_x_2+=     0.00000000082 * cos( 5.90914298142 +    10177.25767953360*t);
+   emb_x_2+=     0.00000000071 * cos( 0.59970099558 +    14314.16811304980*t);
+   emb_x_2+=     0.00000000092 * cos( 1.98240564470 +    13916.01910964160*t);
+   emb_x_2+=     0.00000000091 * cos( 3.10960444681 +     6923.95345737360*t);
+   emb_x_2+=     0.00000000082 * cos( 4.00218490582 +    31415.37924995700*t);
+   emb_x_2+=     0.00000000061 * cos( 4.23984291228 +     8273.82086703240*t);
+   emb_x_2+=     0.00000000055 * cos( 1.41429091428 +     5643.17856367740*t);
+   emb_x_2+=     0.00000000057 * cos( 4.71394484871 +      639.89728631400*t);
+   emb_x_2+=     0.00000000054 * cos( 5.63153101737 +    16496.36139620240*t);
+   emb_x_2+=     0.00000000055 * cos( 3.83548544187 +    13341.67431130680*t);
+   emb_x_2+=     0.00000000052 * cos( 4.39839251424 +     3634.62102451840*t);
+   emb_x_2+=     0.00000000056 * cos( 1.31106242697 +    15720.83878487840*t);
+   emb_x_2+=     0.00000000050 * cos( 4.18716262207 +    12573.26524698360*t);
+   emb_x_2+=     0.00000000064 * cos( 3.35781526441 +     4732.03062734340*t);
+   emb_x_2+=     0.00000000066 * cos( 5.49585898905 +     2787.04302385740*t);
+   emb_x_2+=     0.00000000048 * cos( 2.02353005537 +    13362.44970679920*t);
+   emb_x_2+=     0.00000000049 * cos( 2.48296676043 +    14143.49524243060*t);
+   emb_x_2+=     0.00000000053 * cos( 4.12375134867 +    14919.01785375460*t);
+   emb_x_2+=     0.00000000047 * cos( 3.82959048333 +    12964.30070339100*t);
+   emb_x_2+=     0.00000000042 * cos( 0.38069707842 +    24072.92146977640*t);
+   emb_x_2+=     0.00000000051 * cos( 1.26573618407 +     6195.76764545159*t);
+   emb_x_2+=     0.00000000051 * cos( 5.38279738858 +     6370.38405453120*t);
+   emb_x_2+=     0.00000000048 * cos( 0.59989438773 +     2118.76386037840*t);
+   emb_x_2+=     0.00000000041 * cos( 2.49731951011 +    18319.53658487960*t);
+   emb_x_2+=     0.00000000041 * cos( 4.09328168773 +     3185.19202726560*t);
+   emb_x_2+=     0.00000000039 * cos( 4.27641929836 +    20426.57109242200*t);
+   emb_x_2+=     0.00000000040 * cos( 6.05039617409 +     3496.03282613400*t);
+   emb_x_2+=     0.00000000044 * cos( 1.67103342381 +     8662.24032356300*t);
+   emb_x_2+=     0.00000000035 * cos( 3.60804630271 +    18875.52586977400*t);
+   emb_x_2+=     0.00000000035 * cos( 6.09077128153 +    10419.98628350760*t);
+   emb_x_2+=     0.00000000033 * cos( 0.53088739877 +    10213.28554621100*t);
+   emb_x_2+=     0.00000000031 * cos( 0.64652215197 +    23013.53953958720*t);
+   emb_x_2+=     0.00000000035 * cos( 0.91150627653 +     9779.10867612540*t);
+   emb_x_2+=     0.00000000030 * cos( 4.39715700541 +    12808.88030395680*t);
+   emb_x_2+=     0.00000000031 * cos( 5.53053428232 +    14712.31711645800*t);
+   emb_x_2+=     0.00000000041 * cos( 5.62875780162 +    17789.84561978500*t);
+   emb_x_2+=     0.00000000031 * cos( 0.76672410297 +    22779.43724619380*t);
+   emb_x_2+=     0.00000000028 * cos( 2.34481351815 +     5429.87946823940*t);
+   emb_x_2+=     0.00000000027 * cos( 4.68961132588 +    16200.77272450120*t);
+   emb_x_2+=     0.00000000028 * cos( 3.91086715449 +    11015.10647733480*t);
+   emb_x_2+=     0.00000000026 * cos( 2.29393980670 +     9623.68827669120*t);
+   emb_x_2+=     0.00000000026 * cos( 4.35459376595 +     2942.46342329160*t);
+   emb_x_2+=     0.00000000026 * cos( 1.75234208326 +    18209.33026366019*t);
+   emb_x_2+=     0.00000000025 * cos( 6.04606961258 +     9917.69687450980*t);
+   emb_x_2+=     0.00000000025 * cos( 2.68813215214 +     3583.34103067380*t);
+   emb_x_2+=     0.00000000026 * cos( 0.90493478768 +    12012.58229714040*t);
+   emb_x_2+=     0.00000000023 * cos( 2.29788589227 +     9380.95967271720*t);
+   emb_x_2+=     0.00000000023 * cos( 2.49168005831 +    12559.03815298200*t);
+   emb_x_2+=     0.00000000023 * cos( 2.26727226510 +    26709.64694241340*t);
+   emb_x_2+=     0.00000000021 * cos( 2.38334945179 +    17256.63153634140*t);
+   emb_x_2+=     0.00000000023 * cos( 2.53410224481 +    19651.04848109800*t);
+   emb_x_2+=     0.00000000019 * cos( 1.93058658449 +    11926.25441366880*t);
+   emb_x_2+=     0.00000000018 * cos( 2.28255799994 +    11614.43329373220*t);
+   emb_x_2+=     0.00000000019 * cos( 3.66100662326 +    30639.85663863300*t);
+   emb_x_2+=     0.00000000022 * cos( 3.76382094449 +    17298.18232732620*t);
+   emb_x_2+=     0.00000000019 * cos( 5.39096925285 +    16460.33352952499*t);
+   emb_x_2+=     0.00000000018 * cos( 3.19936496001 +    18073.70493865020*t);
+   emb_x_2+=     0.00000000017 * cos( 3.95703025000 +     8982.81066930900*t);
+   emb_x_2+=     0.00000000021 * cos( 3.77368081700 +    16858.48253293320*t);
+   emb_x_2+=     0.00000000021 * cos( 4.47047819380 +    12036.46073488820*t);
+   emb_x_2+=     0.00000000015 * cos( 0.65562803693 +    16062.18452611680*t);
+   emb_x_2+=     0.00000000014 * cos( 4.12927933278 +     2352.86615377180*t);
+   emb_x_2+=     0.00000000015 * cos( 4.21849500083 +    20199.09495963300*t);
+   emb_x_2+=     0.00000000014 * cos( 0.15947511180 +    22483.84857449259*t);
+   emb_x_2+=     0.00000000019 * cos( 2.81419648175 +     8671.96987044060*t);
+   emb_x_2+=     0.00000000016 * cos( 0.41160329502 +    23543.23050468179*t);
+   emb_x_2+=     0.00000000012 * cos( 2.70990889006 +    17996.03116822220*t);
+   emb_x_2+=     0.00000000012 * cos( 0.72222975090 +    19402.79695281660*t);
+   emb_x_2+=     0.00000000012 * cos( 4.37880728616 +    10818.13528691580*t);
+   emb_x_2+=     0.00000000015 * cos( 0.57625721169 +    12567.13202105100*t);
+   emb_x_2+=     0.00000000011 * cos( 1.78472830783 +    28766.92442448400*t);
+   emb_x_2+=     0.00000000011 * cos( 4.93327896664 +    18451.07854656599*t);
+   emb_x_2+=     0.00000000011 * cos( 5.43587992260 +    14945.31617355440*t);
+   emb_x_2+=     0.00000000011 * cos( 2.24607508630 +    35371.88726597640*t);
+   emb_x_2+=     0.00000000011 * cos( 5.02228901732 +    22805.73556599360*t);
+   emb_x_2+=     0.00000000010 * cos( 2.16930383321 +    15664.03552270859*t);
+   emb_x_2+=     0.00000000014 * cos( 1.93314444705 +    21228.39202354580*t);
+   emb_x_2+=     0.00000000010 * cos( 6.05310595892 +    23141.55838292460*t);
+   emb_x_2+=     0.00000000010 * cos( 2.55645282698 +    22345.26037610820*t);
+   emb_x_2+=     0.00000000009 * cos( 1.49155179482 +    29826.30635467320*t);
+   emb_x_2+=     0.00000000010 * cos( 0.18358815992 +    12779.45079542080*t);
+   emb_x_2+=     0.00000000011 * cos( 4.05473090705 +    23581.25817731760*t);
+   emb_x_2+=     0.00000000011 * cos( 0.39656469738 +    25158.60171976540*t);
+   emb_x_2=emb_x_2 * t * t;
+
+   double emb_x_3=0.0;
+   emb_x_3+=     0.00000175213 * cos( 3.14159265359 +        0.00000000000*t);
+   emb_x_3+=     0.00000072337 * cos( 2.89303952476 +    12566.15169998280*t);
+   emb_x_3+=     0.00000008364 * cos( 3.85500954096 +     6283.07584999140*t);
+   emb_x_3+=     0.00000005040 * cos( 2.53152989786 +    18849.22754997420*t);
+   emb_x_3+=     0.00000001481 * cos( 2.23672235234 +     6438.49624942560*t);
+   emb_x_3+=     0.00000001481 * cos( 4.41181122031 +     6127.65545055720*t);
+   emb_x_3+=     0.00000000644 * cos( 4.88567904780 +     6279.55273164240*t);
+   emb_x_3+=     0.00000000644 * cos( 1.76285452485 +     6286.59896834040*t);
+   emb_x_3+=     0.00000000318 * cos( 3.63840310183 +     6040.34724601740*t);
+   emb_x_3+=     0.00000000318 * cos( 3.01013047082 +     6525.80445396540*t);
+   emb_x_3+=     0.00000000220 * cos( 2.34166945312 +    25132.30339996560*t);
+   emb_x_3+=     0.00000000201 * cos( 0.82714760810 +     6836.64525283380*t);
+   emb_x_3+=     0.00000000201 * cos( 5.82138596455 +     5729.50644714900*t);
+   emb_x_3+=     0.00000000036 * cos( 1.41202770454 +    12569.67481833180*t);
+   emb_x_3+=     0.00000000033 * cos( 4.23554585611 +    12410.73130054860*t);
+   emb_x_3+=     0.00000000018 * cos( 2.32709487755 +    12721.57209941700*t);
+   emb_x_3+=     0.00000000011 * cos( 6.27313981891 +    12592.45001978260*t);
+   emb_x_3+=     0.00000000006 * cos( 2.56028859854 +      242.72860397400*t);
+   emb_x_3=emb_x_3 * t * t * t;
+
+   double emb_x_4=0.0;
+   emb_x_4+=     0.00000004022 * cos( 0.00000000000 +        0.00000000000*t);
+   emb_x_4+=     0.00000001927 * cos( 1.18746233453 +    12566.15169998280*t);
+   emb_x_4+=     0.00000000491 * cos( 5.29890244491 +     6283.07584999140*t);
+   emb_x_4+=     0.00000000210 * cos( 3.74319530048 +     6438.49624942560*t);
+   emb_x_4+=     0.00000000210 * cos( 2.90533827217 +     6127.65545055720*t);
+   emb_x_4+=     0.00000000229 * cos( 0.88072621187 +    18849.22754997420*t);
+   emb_x_4+=     0.00000000021 * cos( 0.18394116301 +     6279.55273164240*t);
+   emb_x_4+=     0.00000000021 * cos( 0.18140710246 +     6286.59896834040*t);
+   emb_x_4+=     0.00000000011 * cos( 0.65552931190 +    25132.30339996560*t);
+   emb_x_4+=     0.00000000005 * cos( 2.68915538708 +    12410.73130054860*t);
+   emb_x_4=emb_x_4 * t * t * t * t;
+
+   double emb_x_5=0.0;
+   emb_x_5+=     0.00000000108 * cos( 0.00000000000 +        0.00000000000*t);
+   emb_x_5+=     0.00000000043 * cos( 6.08789652058 +    12566.15169998280*t);
+   emb_x_5+=     0.00000000025 * cos( 5.33779665346 +     6438.49624942560*t);
+   emb_x_5+=     0.00000000025 * cos( 1.31073691919 +     6127.65545055720*t);
+   emb_x_5+=     0.00000000020 * cos( 0.94333339828 +     6283.07584999140*t);
+   emb_x_5+=     0.00000000008 * cos( 5.52588061142 +    18849.22754997420*t);
+   emb_x_5=emb_x_5 * t * t * t * t * t;
+
+   return emb_x_0+emb_x_1+emb_x_2+emb_x_3+emb_x_4+emb_x_5;
+}
+
+double vsop87a_full::emb_y(double t){
+   double emb_y_0=0.0;
+   emb_y_0+=     0.99989209645 * cos( 0.18265890456 +     6283.07584999140*t);
+   emb_y_0+=     0.02442698841 * cos( 3.14159265359 +        0.00000000000*t);
+   emb_y_0+=     0.00835292314 * cos( 0.13952878991 +    12566.15169998280*t);
+   emb_y_0+=     0.00010466965 * cos( 0.09641690558 +    18849.22754997420*t);
+   emb_y_0+=     0.00002570338 * cos( 5.30103973360 +      529.69096509460*t);
+   emb_y_0+=     0.00002147473 * cos( 2.66253538905 +     1577.34354244780*t);
+   emb_y_0+=     0.00001709219 * cos( 5.20780401071 +     6279.55273164240*t);
+   emb_y_0+=     0.00001707987 * cos( 4.58232858766 +     6286.59896834040*t);
+   emb_y_0+=     0.00001440265 * cos( 1.90068164664 +     2352.86615377180*t);
+   emb_y_0+=     0.00001135092 * cos( 5.27313415220 +     5223.69391980220*t);
+   emb_y_0+=     0.00000934539 * cos( 4.50301201844 +    12036.46073488820*t);
+   emb_y_0+=     0.00000900565 * cos( 1.60563288120 +    10213.28554621100*t);
+   emb_y_0+=     0.00000567126 * cos( 0.58142248753 +     1059.38193018920*t);
+   emb_y_0+=     0.00000744932 * cos( 2.80728871886 +      398.14900340820*t);
+   emb_y_0+=     0.00000639316 * cos( 6.02923915017 +     5753.38488489680*t);
+   emb_y_0+=     0.00000681324 * cos( 0.64729627497 +     4705.73230754360*t);
+   emb_y_0+=     0.00000611347 * cos( 3.81381495286 +     6812.76681508600*t);
+   emb_y_0+=     0.00000450435 * cos( 4.52785572489 +     5884.92684658320*t);
+   emb_y_0+=     0.00000452018 * cos( 5.99167242707 +     6256.77753019160*t);
+   emb_y_0+=     0.00000449968 * cos( 3.79880375595 +     6309.37416979120*t);
+   emb_y_0+=     0.00000551390 * cos( 3.96125249369 +     5507.55323866740*t);
+   emb_y_0+=     0.00000406334 * cos( 5.25616268027 +     6681.22485339960*t);
+   emb_y_0+=     0.00000541273 * cos( 5.49902805917 +      775.52261132400*t);
+   emb_y_0+=     0.00000546360 * cos( 6.17311131785 +    14143.49524243060*t);
+   emb_y_0+=     0.00000507084 * cos( 2.87025193381 +     7860.41939243920*t);
+   emb_y_0+=     0.00000219504 * cos( 2.95216139568 +    11506.76976979360*t);
+   emb_y_0+=     0.00000227937 * cos( 5.95179248814 +     7058.59846131540*t);
+   emb_y_0+=     0.00000227792 * cos( 4.84547074733 +     4694.00295470760*t);
+   emb_y_0+=     0.00000255845 * cos( 0.69454231563 +    12168.00269657460*t);
+   emb_y_0+=     0.00000179242 * cos( 1.40003446021 +      796.29800681640*t);
+   emb_y_0+=     0.00000178280 * cos( 5.11717552231 +     6283.00853968860*t);
+   emb_y_0+=     0.00000178280 * cos( 4.67307255246 +     6283.14316029419*t);
+   emb_y_0+=     0.00000155454 * cos( 0.05340525434 +    25132.30339996560*t);
+   emb_y_0+=     0.00000206257 * cos( 4.28366728882 +    11790.62908865880*t);
+   emb_y_0+=     0.00000149769 * cos( 6.07429023278 +     5486.77784317500*t);
+   emb_y_0+=     0.00000200005 * cos( 2.50144088120 +    17789.84561978500*t);
+   emb_y_0+=     0.00000129006 * cos( 3.64623708634 +     7079.37385680780*t);
+   emb_y_0+=     0.00000128211 * cos( 3.23254821381 +     3738.76143010800*t);
+   emb_y_0+=     0.00000152790 * cos( 5.58120800450 +      213.29909543800*t);
+   emb_y_0+=     0.00000118725 * cos( 5.45361490488 +     9437.76293488700*t);
+   emb_y_0+=     0.00000078419 * cos( 0.28253093914 +     8827.39026987480*t);
+   emb_y_0+=     0.00000073513 * cos( 1.12433371880 +     1589.07289528380*t);
+   emb_y_0+=     0.00000065694 * cos( 2.08898587059 +    11769.85369316640*t);
+   emb_y_0+=     0.00000063240 * cos( 0.67044491570 +     6262.30045449900*t);
+   emb_y_0+=     0.00000062961 * cos( 2.83649755996 +     6303.85124548380*t);
+   emb_y_0+=     0.00000063616 * cos( 0.49038985445 +        7.11354700080*t);
+   emb_y_0+=     0.00000080740 * cos( 0.92225805400 +     6069.77675455340*t);
+   emb_y_0+=     0.00000083511 * cos( 4.63223527298 +     3340.61242669980*t);
+   emb_y_0+=     0.00000060039 * cos( 1.82878888713 +     4136.91043351620*t);
+   emb_y_0+=     0.00000075626 * cos( 2.58860726122 +     6496.37494542940*t);
+   emb_y_0+=     0.00000058089 * cos( 0.55165509745 +    10977.07880469900*t);
+   emb_y_0+=     0.00000053747 * cos( 6.27553948416 +     1194.44701022460*t);
+   emb_y_0+=     0.00000057574 * cos( 1.10802134270 +     6282.09552892320*t);
+   emb_y_0+=     0.00000057574 * cos( 2.39891957636 +     6284.05617105960*t);
+   emb_y_0+=     0.00000051141 * cos( 6.02458046766 +     6290.18939699220*t);
+   emb_y_0+=     0.00000050679 * cos( 3.77133779892 +     6275.96230299060*t);
+   emb_y_0+=     0.00000049613 * cos( 1.07290390127 +     6127.65545055720*t);
+   emb_y_0+=     0.00000047934 * cos( 2.43176952238 +     6438.49624942560*t);
+   emb_y_0+=     0.00000061133 * cos( 5.34335721244 +     3154.68708489560*t);
+   emb_y_0+=     0.00000048425 * cos( 0.05833183710 +      801.82093112380*t);
+   emb_y_0+=     0.00000047359 * cos( 4.53607945311 +     3128.38876509580*t);
+   emb_y_0+=     0.00000048379 * cos( 0.25856696781 +     7084.89678111520*t);
+   emb_y_0+=     0.00000051293 * cos( 4.21367834558 +    12352.85260454480*t);
+   emb_y_0+=     0.00000037539 * cos( 1.60040603609 +     8429.24126646660*t);
+   emb_y_0+=     0.00000049347 * cos( 2.15740113992 +     2544.31441988340*t);
+   emb_y_0+=     0.00000039555 * cos( 3.35943284176 +     5481.25491886760*t);
+   emb_y_0+=     0.00000048913 * cos( 5.16531049040 +     9225.53927328300*t);
+   emb_y_0+=     0.00000042135 * cos( 2.06532560040 +    17260.15465469040*t);
+   emb_y_0+=     0.00000037954 * cos( 1.47005640108 +      426.59819087600*t);
+   emb_y_0+=     0.00000031140 * cos( 3.61579568733 +     2146.16541647520*t);
+   emb_y_0+=     0.00000032603 * cos( 4.41806626166 +     4164.31198961300*t);
+   emb_y_0+=     0.00000039728 * cos( 4.70048863454 +     5856.47765911540*t);
+   emb_y_0+=     0.00000028515 * cos( 4.53853388964 +    12569.67481833180*t);
+   emb_y_0+=     0.00000036368 * cos( 1.33416490018 +    18073.70493865020*t);
+   emb_y_0+=     0.00000025032 * cos( 3.74982021361 +       26.29831979980*t);
+   emb_y_0+=     0.00000030329 * cos( 0.42324614391 +     4535.05943692440*t);
+   emb_y_0+=     0.00000026741 * cos( 5.97172882162 +     3930.20969621960*t);
+   emb_y_0+=     0.00000028166 * cos( 5.19366222912 +    12562.62858163380*t);
+   emb_y_0+=     0.00000024843 * cos( 1.00265326013 +     5088.62883976680*t);
+   emb_y_0+=     0.00000021642 * cos( 3.06444526478 +    12559.03815298200*t);
+   emb_y_0+=     0.00000025701 * cos( 0.75895390818 +    13367.97263110660*t);
+   emb_y_0+=     0.00000020498 * cos( 3.44435384536 +    11371.70468975820*t);
+   emb_y_0+=     0.00000017888 * cos( 5.01504194498 +     1748.01641306700*t);
+   emb_y_0+=     0.00000023246 * cos( 4.61910038205 +    18319.53658487960*t);
+   emb_y_0+=     0.00000017077 * cos( 5.29722216277 +     4933.20844033260*t);
+   emb_y_0+=     0.00000021070 * cos( 2.22959477131 +     7477.52286021600*t);
+   emb_y_0+=     0.00000017512 * cos( 6.16192589291 +    10447.38783960440*t);
+   emb_y_0+=     0.00000020024 * cos( 2.96823180998 +     8031.09226305840*t);
+   emb_y_0+=     0.00000013831 * cos( 3.02450528624 +     8635.94200376320*t);
+   emb_y_0+=     0.00000014095 * cos( 6.03707103538 +     2942.46342329160*t);
+   emb_y_0+=     0.00000017861 * cos( 4.86809570283 +     1592.59601363280*t);
+   emb_y_0+=     0.00000013605 * cos( 3.62641145030 +    13095.84266507740*t);
+   emb_y_0+=     0.00000015078 * cos( 1.55132224732 +    16496.36139620240*t);
+   emb_y_0+=     0.00000017985 * cos( 2.07741211308 +    17298.18232732620*t);
+   emb_y_0+=     0.00000014386 * cos( 6.14898732933 +    20426.57109242200*t);
+   emb_y_0+=     0.00000012227 * cos( 4.38359423735 +     7632.94325965020*t);
+   emb_y_0+=     0.00000013376 * cos( 2.03616887071 +    12139.55350910680*t);
+   emb_y_0+=     0.00000011583 * cos( 3.88971333608 +     5331.35744374080*t);
+   emb_y_0+=     0.00000013558 * cos( 1.44666331503 +    16730.46368959580*t);
+   emb_y_0+=     0.00000014332 * cos( 4.45678693729 +     7342.45778018060*t);
+   emb_y_0+=     0.00000012375 * cos( 1.53862289477 +     6279.48542133960*t);
+   emb_y_0+=     0.00000012366 * cos( 1.96831814872 +     6286.66627864320*t);
+   emb_y_0+=     0.00000012446 * cos( 5.78094918529 +    15110.46611986620*t);
+   emb_y_0+=     0.00000010989 * cos( 0.62040343662 +    10988.80815753500*t);
+   emb_y_0+=     0.00000012556 * cos( 2.59194808538 +    15720.83878487840*t);
+   emb_y_0+=     0.00000011583 * cos( 2.48165809843 +     5729.50644714900*t);
+   emb_y_0+=     0.00000009023 * cos( 3.85655029499 +     9623.68827669120*t);
+   emb_y_0+=     0.00000009121 * cos( 5.83006047147 +     7234.79425624200*t);
+   emb_y_0+=     0.00000008869 * cos( 3.47611265290 +     6148.01076995600*t);
+   emb_y_0+=     0.00000008614 * cos( 0.03081589303 +     6418.14093002680*t);
+   emb_y_0+=     0.00000007911 * cos( 1.58348238359 +     2118.76386037840*t);
+   emb_y_0+=     0.00000008293 * cos( 0.85407021371 +    14712.31711645800*t);
+   emb_y_0+=     0.00000010219 * cos( 0.13075179540 +     1349.86740965880*t);
+   emb_y_0+=     0.00000010126 * cos( 0.99630008630 +     6836.64525283380*t);
+   emb_y_0+=     0.00000010480 * cos( 1.38209301929 +     5999.21653112620*t);
+   emb_y_0+=     0.00000007353 * cos( 3.83220307326 +     6040.34724601740*t);
+   emb_y_0+=     0.00000009820 * cos( 2.12476655144 +     6566.93516885660*t);
+   emb_y_0+=     0.00000007018 * cos( 5.94659180519 +     6525.80445396540*t);
+   emb_y_0+=     0.00000006874 * cos( 4.57503512461 +     7238.67559160000*t);
+   emb_y_0+=     0.00000007545 * cos( 5.94843444430 +    12539.85338018300*t);
+   emb_y_0+=     0.00000009540 * cos( 3.49613271706 +    21228.39202354580*t);
+   emb_y_0+=     0.00000007202 * cos( 4.31311898936 +     6245.04817735560*t);
+   emb_y_0+=     0.00000007144 * cos( 5.47703560836 +     6321.10352262720*t);
+   emb_y_0+=     0.00000006627 * cos( 5.41373102778 +     5327.47610838280*t);
+   emb_y_0+=     0.00000008544 * cos( 1.82829301857 +    11015.10647733480*t);
+   emb_y_0+=     0.00000008602 * cos( 5.43953575310 +     5643.17856367740*t);
+   emb_y_0+=     0.00000006158 * cos( 5.65852675233 +    12964.30070339100*t);
+   emb_y_0+=     0.00000007414 * cos( 3.46720433621 +    23543.23050468179*t);
+   emb_y_0+=     0.00000007317 * cos( 2.01664477821 +      536.80451209540*t);
+   emb_y_0+=     0.00000006352 * cos( 3.46461304398 +     1990.74501704100*t);
+   emb_y_0+=     0.00000006251 * cos( 1.52860898333 +      951.71840625060*t);
+   emb_y_0+=     0.00000005828 * cos( 1.37555383064 +      955.59974160860*t);
+   emb_y_0+=     0.00000006389 * cos( 2.40580071461 +    24072.92146977640*t);
+   emb_y_0+=     0.00000006976 * cos( 4.76626361488 +    10973.55568635000*t);
+   emb_y_0+=     0.00000006916 * cos( 5.39658071678 +      522.57741809380*t);
+   emb_y_0+=     0.00000007459 * cos( 2.71133898870 +    22003.91463486980*t);
+   emb_y_0+=     0.00000005376 * cos( 2.34968505500 +    14314.16811304980*t);
+   emb_y_0+=     0.00000006852 * cos( 4.33781274406 +      155.42039943420*t);
+   emb_y_0+=     0.00000007527 * cos( 1.94392127074 +    18422.62935909819*t);
+   emb_y_0+=     0.00000006222 * cos( 0.67296135144 +    18451.07854656599*t);
+   emb_y_0+=     0.00000005836 * cos( 0.75921238848 +      639.89728631400*t);
+   emb_y_0+=     0.00000006487 * cos( 2.22135874272 +     5216.58037280140*t);
+   emb_y_0+=     0.00000006251 * cos( 2.88399458161 +     6288.59877429880*t);
+   emb_y_0+=     0.00000006251 * cos( 0.62294633745 +     6277.55292568400*t);
+   emb_y_0+=     0.00000007004 * cos( 3.57492830242 +      103.09277421860*t);
+   emb_y_0+=     0.00000005337 * cos( 1.97127813671 +     4690.47983635860*t);
+   emb_y_0+=     0.00000005605 * cos( 1.92150700481 +     1551.04522264800*t);
+   emb_y_0+=     0.00000006210 * cos( 0.81866451118 +     5650.29211067820*t);
+   emb_y_0+=     0.00000006064 * cos( 5.19708779870 +     5230.80746680300*t);
+   emb_y_0+=     0.00000004596 * cos( 2.92936593195 +      553.56940284240*t);
+   emb_y_0+=     0.00000006332 * cos( 2.83725098708 +    26087.90314157420*t);
+   emb_y_0+=     0.00000005016 * cos( 1.46006858313 +     4732.03062734340*t);
+   emb_y_0+=     0.00000005414 * cos( 1.42882325200 +     6179.98307577280*t);
+   emb_y_0+=     0.00000005413 * cos( 2.08325983447 +     6386.16862421000*t);
+   emb_y_0+=     0.00000005711 * cos( 2.75630374926 +     6915.85958930460*t);
+   emb_y_0+=     0.00000004270 * cos( 0.81269903423 +     7875.67186362420*t);
+   emb_y_0+=     0.00000004017 * cos( 2.21501265909 +     6254.62666252360*t);
+   emb_y_0+=     0.00000004071 * cos( 3.99159158863 +     3634.62102451840*t);
+   emb_y_0+=     0.00000003993 * cos( 1.29192500625 +     6311.52503745920*t);
+   emb_y_0+=     0.00000004858 * cos( 1.27914374675 +     5760.49843189760*t);
+   emb_y_0+=     0.00000005274 * cos( 4.91687036691 +    25158.60171976540*t);
+   emb_y_0+=     0.00000004120 * cos( 1.72568521475 +     6709.67404086740*t);
+   emb_y_0+=     0.00000004410 * cos( 2.23536264210 +     6805.65326808520*t);
+   emb_y_0+=     0.00000003174 * cos( 3.62030131501 +     1066.49547719000*t);
+   emb_y_0+=     0.00000003064 * cos( 2.46605561478 +     5863.59120611620*t);
+   emb_y_0+=     0.00000003083 * cos( 5.68731909779 +     9917.69687450980*t);
+   emb_y_0+=     0.00000003385 * cos( 3.81076637101 +    13916.01910964160*t);
+   emb_y_0+=     0.00000002979 * cos( 4.62987940723 +    12566.21901028560*t);
+   emb_y_0+=     0.00000003980 * cos( 5.90917210624 +    13341.67431130680*t);
+   emb_y_0+=     0.00000003290 * cos( 0.95752698537 +    16200.77272450120*t);
+   emb_y_0+=     0.00000003885 * cos( 3.94360149823 +       74.78159856730*t);
+   emb_y_0+=     0.00000003060 * cos( 1.63142398497 +    18052.92954315780*t);
+   emb_y_0+=     0.00000003745 * cos( 3.65023712104 +     5966.68398033480*t);
+   emb_y_0+=     0.00000003592 * cos( 4.03843804285 +     6208.29425142410*t);
+   emb_y_0+=     0.00000003581 * cos( 5.75234129258 +     6357.85744855870*t);
+   emb_y_0+=     0.00000003626 * cos( 6.13871061291 +     6599.46771964800*t);
+   emb_y_0+=     0.00000002927 * cos( 1.79201822142 +    12721.57209941700*t);
+   emb_y_0+=     0.00000003625 * cos( 1.43112641053 +     6019.99192661860*t);
+   emb_y_0+=     0.00000002589 * cos( 2.27559640821 +    12592.45001978260*t);
+   emb_y_0+=     0.00000002558 * cos( 6.08151329647 +    10575.40668294180*t);
+   emb_y_0+=     0.00000002670 * cos( 2.64370887209 +        3.52311834900*t);
+   emb_y_0+=     0.00000002611 * cos( 0.86180315523 +     6702.56049386660*t);
+   emb_y_0+=     0.00000002535 * cos( 0.01231078888 +    31415.37924995700*t);
+   emb_y_0+=     0.00000003434 * cos( 2.07577333297 +     6546.15977336420*t);
+   emb_y_0+=     0.00000002378 * cos( 2.05651519036 +     2388.89402044920*t);
+   emb_y_0+=     0.00000002979 * cos( 5.07398237708 +    12566.08438968000*t);
+   emb_y_0+=     0.00000002716 * cos( 1.49855840728 +    12029.34718788740*t);
+   emb_y_0+=     0.00000002215 * cos( 3.69574321309 +    13362.44970679920*t);
+   emb_y_0+=     0.00000003049 * cos( 0.05228873308 +    29088.81141598500*t);
+   emb_y_0+=     0.00000002784 * cos( 3.37954165171 +    14945.31617355440*t);
+   emb_y_0+=     0.00000002284 * cos( 4.38146814547 +     6076.89030155420*t);
+   emb_y_0+=     0.00000002529 * cos( 4.37629589914 +    12043.57428188900*t);
+   emb_y_0+=     0.00000002015 * cos( 0.91352603582 +     6262.72053059260*t);
+   emb_y_0+=     0.00000002218 * cos( 5.40562909521 +     6489.26139842860*t);
+   emb_y_0+=     0.00000002005 * cos( 2.59344263770 +     6303.43116939020*t);
+   emb_y_0+=     0.00000002681 * cos( 5.85620039558 +     4590.91018048900*t);
+   emb_y_0+=     0.00000002011 * cos( 5.69828891426 +     4701.11650170840*t);
+   emb_y_0+=     0.00000002440 * cos( 4.55844746985 +      149.56319713460*t);
+   emb_y_0+=     0.00000002600 * cos( 3.93109819647 +    19651.04848109800*t);
+   emb_y_0+=     0.00000002197 * cos( 5.22411628121 +    13517.87010623340*t);
+   emb_y_0+=     0.00000001893 * cos( 4.09716259686 +     6286.95718534940*t);
+   emb_y_0+=     0.00000001893 * cos( 5.69296362938 +     6279.19451463340*t);
+   emb_y_0+=     0.00000002402 * cos( 4.84503253270 +     1162.47470440780*t);
+   emb_y_0+=     0.00000001778 * cos( 2.79119893202 +     4292.33083295040*t);
+   emb_y_0+=     0.00000001966 * cos( 1.82224928185 +     4686.88940770680*t);
+   emb_y_0+=     0.00000002327 * cos( 2.83472726504 +    10021.83728009940*t);
+   emb_y_0+=     0.00000002163 * cos( 6.13036765564 +     6037.24420376200*t);
+   emb_y_0+=     0.00000001841 * cos( 3.67947373190 +      242.72860397400*t);
+   emb_y_0+=     0.00000001672 * cos( 0.33097080655 +    13119.72110282519*t);
+   emb_y_0+=     0.00000002250 * cos( 0.82649282051 +    12669.24447420140*t);
+   emb_y_0+=     0.00000001627 * cos( 4.77445737868 +    13521.75144159140*t);
+   emb_y_0+=     0.00000001888 * cos( 2.87850805387 +     8662.24032356300*t);
+   emb_y_0+=     0.00000002052 * cos( 3.65927016475 +     6528.90749622080*t);
+   emb_y_0+=     0.00000001774 * cos( 0.51400448882 +     6819.88036208680*t);
+   emb_y_0+=     0.00000001544 * cos( 1.33911061474 +      110.20632121940*t);
+   emb_y_0+=     0.00000001756 * cos( 2.96482886554 +     5746.27133789600*t);
+   emb_y_0+=     0.00000001543 * cos( 4.39018655159 +     6272.43918464160*t);
+   emb_y_0+=     0.00000001543 * cos( 5.39993967465 +     6293.71251534120*t);
+   emb_y_0+=     0.00000001952 * cos( 3.09643550846 +    23013.53953958720*t);
+   emb_y_0+=     0.00000001484 * cos( 2.15301976837 +       20.77539549240*t);
+   emb_y_0+=     0.00000001558 * cos( 2.79679951671 +    17654.78053974960*t);
+   emb_y_0+=     0.00000001440 * cos( 0.69014829975 +     1052.26838318840*t);
+   emb_y_0+=     0.00000001648 * cos( 3.98135416194 +      632.78373931320*t);
+   emb_y_0+=     0.00000001795 * cos( 1.47161886169 +    33019.02111220460*t);
+   emb_y_0+=     0.00000001488 * cos( 4.10377440731 +       76.26607127560*t);
+   emb_y_0+=     0.00000001349 * cos( 4.18018072818 +     6016.46880826960*t);
+   emb_y_0+=     0.00000001477 * cos( 2.56925128481 +    12779.45079542080*t);
+   emb_y_0+=     0.00000001252 * cos( 6.26830198321 +    11499.65622279280*t);
+   emb_y_0+=     0.00000001491 * cos( 4.65893645946 +     9411.46461508720*t);
+   emb_y_0+=     0.00000001471 * cos( 4.18695456530 +    18635.92845453620*t);
+   emb_y_0+=     0.00000001263 * cos( 5.62967869388 +    12573.26524698360*t);
+   emb_y_0+=     0.00000001385 * cos( 2.65776727703 +     2379.16447357160*t);
+   emb_y_0+=     0.00000001272 * cos( 5.60977857361 +     6549.68289171320*t);
+   emb_y_0+=     0.00000001578 * cos( 4.10046385407 +    25934.12433108940*t);
+   emb_y_0+=     0.00000001557 * cos( 4.78614835776 +    18875.52586977400*t);
+   emb_y_0+=     0.00000001358 * cos( 0.33781842226 +     5120.60114558360*t);
+   emb_y_0+=     0.00000001505 * cos( 2.70837129901 +    18209.33026366019*t);
+   emb_y_0+=     0.00000001460 * cos( 3.40156973977 +      245.83164622940*t);
+   emb_y_0+=     0.00000001242 * cos( 1.27359177098 +    24356.78078864160*t);
+   emb_y_0+=     0.00000001464 * cos( 3.73000165992 +       38.13303563780*t);
+   emb_y_0+=     0.00000001028 * cos( 0.93950338023 +     6112.40297937220*t);
+   emb_y_0+=     0.00000001014 * cos( 4.36072926630 +    17256.63153634140*t);
+   emb_y_0+=     0.00000001312 * cos( 5.81330496460 +    12491.37010141550*t);
+   emb_y_0+=     0.00000000992 * cos( 2.56738139157 +     6453.74872061060*t);
+   emb_y_0+=     0.00000000988 * cos( 5.69171453384 +     8273.82086703240*t);
+   emb_y_0+=     0.00000001191 * cos( 1.85928407242 +     6290.12208668940*t);
+   emb_y_0+=     0.00000001191 * cos( 1.64765684664 +     6276.02961329339*t);
+   emb_y_0+=     0.00000000978 * cos( 3.16773119822 +     6247.04798331399*t);
+   emb_y_0+=     0.00000001223 * cos( 6.22694723537 +     5429.87946823940*t);
+   emb_y_0+=     0.00000000971 * cos( 0.33921097188 +     6319.10371666880*t);
+   emb_y_0+=     0.00000000962 * cos( 1.06488912173 +    12565.17137891460*t);
+   emb_y_0+=     0.00000001025 * cos( 0.64688768472 +    12545.37630449040*t);
+   emb_y_0+=     0.00000001234 * cos( 2.50427756599 +    22483.84857449259*t);
+   emb_y_0+=     0.00000001134 * cos( 5.76287831677 +    17782.73207278420*t);
+   emb_y_0+=     0.00000000944 * cos( 5.49712708636 +     6306.95428773920*t);
+   emb_y_0+=     0.00000000944 * cos( 4.29299913988 +     6259.19741224360*t);
+   emb_y_0+=     0.00000001092 * cos( 2.35048657695 +    17796.95916678580*t);
+   emb_y_0+=     0.00000000995 * cos( 1.11794413013 +    10177.25767953360*t);
+   emb_y_0+=     0.00000000894 * cos( 2.94681683620 +      206.18554843720*t);
+   emb_y_0+=     0.00000000860 * cos( 6.19327755003 +     1692.16566950240*t);
+   emb_y_0+=     0.00000000942 * cos( 6.13301518193 +     6267.82337880639*t);
+   emb_y_0+=     0.00000000942 * cos( 3.65711104430 +     6298.32832117640*t);
+   emb_y_0+=     0.00000001045 * cos( 5.17292899506 +    15508.61512327440*t);
+   emb_y_0+=     0.00000000850 * cos( 6.16443178844 +     6393.28217121080*t);
+   emb_y_0+=     0.00000000873 * cos( 2.86766306801 +      949.17560896980*t);
+   emb_y_0+=     0.00000000896 * cos( 0.83013691514 +     3903.91137641980*t);
+   emb_y_0+=     0.00000001023 * cos( 5.42813246189 +      220.41264243880*t);
+   emb_y_0+=     0.00000000856 * cos( 1.36181993943 +     4804.20927592700*t);
+   emb_y_0+=     0.00000000817 * cos( 3.58855573226 +     6172.86952877200*t);
+   emb_y_0+=     0.00000001081 * cos( 2.89245979526 +    36949.23080842420*t);
+   emb_y_0+=     0.00000000779 * cos( 2.01260281552 +     2648.45482547300*t);
+   emb_y_0+=     0.00000000962 * cos( 2.35578735539 +    12567.13202105100*t);
+   emb_y_0+=     0.00000000990 * cos( 4.97210778261 +    11403.67699557500*t);
+   emb_y_0+=     0.00000000739 * cos( 3.39896198161 +    11926.25441366880*t);
+   emb_y_0+=     0.00000000929 * cos( 0.64994792213 +     2787.04302385740*t);
+   emb_y_0+=     0.00000000973 * cos( 6.19358824421 +    22805.73556599360*t);
+   emb_y_0+=     0.00000000743 * cos( 2.22828430376 +     7834.12107263940*t);
+   emb_y_0+=     0.00000000905 * cos( 1.26234908246 +     6281.59137728310*t);
+   emb_y_0+=     0.00000000905 * cos( 2.24459183660 +     6284.56032269970*t);
+   emb_y_0+=     0.00000000798 * cos( 1.05421393588 +    12410.73130054860*t);
+   emb_y_0+=     0.00000000813 * cos( 5.23283372183 +    12416.58850284820*t);
+   emb_y_0+=     0.00000000869 * cos( 1.74785663537 +     6133.51265285680*t);
+   emb_y_0+=     0.00000000756 * cos( 5.40063116731 +     4061.21921539440*t);
+   emb_y_0+=     0.00000000616 * cos( 4.13389755770 +     1596.18644228460*t);
+   emb_y_0+=     0.00000000612 * cos( 5.89184406078 +    16858.48253293320*t);
+   emb_y_0+=     0.00000000814 * cos( 3.82004246488 +     3894.18182954220*t);
+   emb_y_0+=     0.00000000785 * cos( 2.44918664136 +     4907.30205014560*t);
+   emb_y_0+=     0.00000000788 * cos( 1.86746379419 +      846.08283475120*t);
+   emb_y_0+=     0.00000000611 * cos( 6.24387943735 +      853.19638175200*t);
+   emb_y_0+=     0.00000000570 * cos( 2.97586328906 +    12586.92709547520*t);
+   emb_y_0+=     0.00000000561 * cos( 5.85814246596 +    12323.42309600880*t);
+   emb_y_0+=     0.00000000646 * cos( 5.91116942848 +    52175.80628314840*t);
+   emb_y_0+=     0.00000000590 * cos( 0.47736597307 +    15671.08175940660*t);
+   emb_y_0+=     0.00000000555 * cos( 2.86698301793 +    11513.88331679440*t);
+   emb_y_0+=     0.00000000536 * cos( 4.49580958272 +    18852.75066832320*t);
+   emb_y_0+=     0.00000000737 * cos( 5.41358302312 +    23581.25817731760*t);
+   emb_y_0+=     0.00000000532 * cos( 3.01682571392 +    18842.11400297339*t);
+   emb_y_0+=     0.00000000552 * cos( 6.08260944683 +    11933.36796066960*t);
+   emb_y_0+=     0.00000000599 * cos( 0.54713576844 +     6151.53388830500*t);
+   emb_y_0+=     0.00000000540 * cos( 1.83361074884 +     6279.78949257360*t);
+   emb_y_0+=     0.00000000540 * cos( 1.67333017022 +     6286.36220740920*t);
+   emb_y_0+=     0.00000000513 * cos( 2.13211972766 +    12455.94537876340*t);
+   emb_y_0+=     0.00000000649 * cos( 3.33488856037 +    11764.33076885900*t);
+   emb_y_0+=     0.00000000582 * cos( 2.96012819576 +     6414.61781167780*t);
+   emb_y_0+=     0.00000000608 * cos( 2.37588807672 +    10419.98628350760*t);
+   emb_y_0+=     0.00000000620 * cos( 0.83937590497 +     5326.78669402080*t);
+   emb_y_0+=     0.00000000536 * cos( 2.13145334966 +     7232.25145896120*t);
+   emb_y_0+=     0.00000000538 * cos( 1.33481115516 +     5333.90024102160*t);
+   emb_y_0+=     0.00000000516 * cos( 1.35797709246 +      433.71173787680*t);
+   emb_y_0+=     0.00000000596 * cos( 0.99065349155 +     6062.66320755260*t);
+   emb_y_0+=     0.00000000664 * cos( 4.31358012484 +    40879.44050464380*t);
+   emb_y_0+=     0.00000000607 * cos( 4.79849916312 +     5849.36411211460*t);
+   emb_y_0+=     0.00000000629 * cos( 1.34455816041 +    26735.94526221320*t);
+   emb_y_0+=     0.00000000479 * cos( 6.04755907663 +     5547.19933645960*t);
+   emb_y_0+=     0.00000000565 * cos( 5.49238261982 +     4171.42553661380*t);
+   emb_y_0+=     0.00000000472 * cos( 3.81181804730 +    10969.96525769820*t);
+   emb_y_0+=     0.00000000462 * cos( 3.53729128828 +     3104.93005942380*t);
+   emb_y_0+=     0.00000000551 * cos( 2.46794798020 +     6503.48849243020*t);
+   emb_y_0+=     0.00000000450 * cos( 0.38613288324 +     4274.51831083240*t);
+   emb_y_0+=     0.00000000527 * cos( 4.67747327047 +    24602.61243487099*t);
+   emb_y_0+=     0.00000000502 * cos( 5.68646136646 +    12489.88562870720*t);
+   emb_y_0+=     0.00000000432 * cos( 5.22277533085 +     9388.00590941520*t);
+   emb_y_0+=     0.00000000530 * cos( 5.14730101928 +    18845.70443162520*t);
+   emb_y_0+=     0.00000000431 * cos( 3.74066314880 +     7018.95236352320*t);
+   emb_y_0+=     0.00000000484 * cos( 0.31541423564 +       14.22709400160*t);
+   emb_y_0+=     0.00000000560 * cos( 0.14432555287 +     1581.95934828300*t);
+   emb_y_0+=     0.00000000408 * cos( 4.98636538552 +     3496.03282613400*t);
+   emb_y_0+=     0.00000000478 * cos( 4.58275550120 +    29296.61538957860*t);
+   emb_y_0+=     0.00000000389 * cos( 1.64726374330 +     6321.20888562920*t);
+   emb_y_0+=     0.00000000389 * cos( 1.85967717576 +     6244.94281435360*t);
+   emb_y_0+=     0.00000000388 * cos( 5.85202122263 +     6297.30294399300*t);
+   emb_y_0+=     0.00000000388 * cos( 3.93810500361 +     6268.84875598980*t);
+   emb_y_0+=     0.00000000515 * cos( 0.09995703640 +    12320.32005375340*t);
+   emb_y_0+=     0.00000000395 * cos( 1.39225625253 +     4157.19844261220*t);
+   emb_y_0+=     0.00000000492 * cos( 4.39757189321 +     1478.86657406440*t);
+   emb_y_0+=     0.00000000423 * cos( 2.45281376602 +     9779.10867612540*t);
+   emb_y_0+=     0.00000000501 * cos( 3.08355690853 +    17157.06188047180*t);
+   emb_y_0+=     0.00000000416 * cos( 2.17399987974 +    13760.59871020740*t);
+   emb_y_0+=     0.00000000503 * cos( 4.02857067171 +      135.06508003540*t);
+   emb_y_0+=     0.00000000393 * cos( 6.27884951137 +     7872.14874527520*t);
+   emb_y_0+=     0.00000000365 * cos( 0.53940975175 +     6922.97313630540*t);
+   emb_y_0+=     0.00000000409 * cos( 6.14577564379 +    10984.19235169980*t);
+   emb_y_0+=     0.00000000371 * cos( 4.19153897517 +    13625.53363017200*t);
+   emb_y_0+=     0.00000000370 * cos( 2.79623025524 +    14919.01785375460*t);
+   emb_y_0+=     0.00000000379 * cos( 5.52650430530 +     3185.19202726560*t);
+   emb_y_0+=     0.00000000420 * cos( 1.76851201552 +    22779.43724619380*t);
+   emb_y_0+=     0.00000000485 * cos( 6.06683096360 +    12528.01866434500*t);
+   emb_y_0+=     0.00000000358 * cos( 1.58658367058 +     6271.34649715540*t);
+   emb_y_0+=     0.00000000358 * cos( 1.92035724848 +     6294.80520282740*t);
+   emb_y_0+=     0.00000000354 * cos( 2.96580857786 +    17267.26820169119*t);
+   emb_y_0+=     0.00000000351 * cos( 0.64808322427 +    11616.97609101300*t);
+   emb_y_0+=     0.00000000388 * cos( 5.20354943335 +      956.28915597060*t);
+   emb_y_0+=     0.00000000409 * cos( 2.03946298626 +    21954.15760939799*t);
+   emb_y_0+=     0.00000000369 * cos( 3.36327405725 +     5436.99301524020*t);
+   emb_y_0+=     0.00000000421 * cos( 3.25654135729 +     7445.55055439920*t);
+   emb_y_0+=     0.00000000361 * cos( 5.53538077166 +    21393.54196985760*t);
+   emb_y_0+=     0.00000000336 * cos( 3.66443056835 +     6287.00800325450*t);
+   emb_y_0+=     0.00000000336 * cos( 6.12569565789 +     6279.14369672830*t);
+   emb_y_0+=     0.00000000315 * cos( 3.47678894225 +    10873.98603048040*t);
+   emb_y_0+=     0.00000000363 * cos( 2.71990145670 +     7335.34423317980*t);
+   emb_y_0+=     0.00000000347 * cos( 5.27040068469 +    17253.04110768959*t);
+   emb_y_0+=     0.00000000349 * cos( 1.04077270827 +      316.39186965660*t);
+   emb_y_0+=     0.00000000395 * cos( 3.82348293170 +     5905.70224207560*t);
+   emb_y_0+=     0.00000000308 * cos( 2.54337606421 +    20995.39296644940*t);
+   emb_y_0+=     0.00000000319 * cos( 5.30613709521 +       38.02767263580*t);
+   emb_y_0+=     0.00000000304 * cos( 4.87224757874 +     5749.86176654780*t);
+   emb_y_0+=     0.00000000414 * cos( 5.73458482925 +    44809.65020086340*t);
+   emb_y_0+=     0.00000000298 * cos( 2.28443357697 +      640.87760738220*t);
+   emb_y_0+=     0.00000000351 * cos( 4.36315900786 +    12345.73905754400*t);
+   emb_y_0+=     0.00000000304 * cos( 3.36893995706 +    29826.30635467320*t);
+   emb_y_0+=     0.00000000395 * cos( 4.78955929514 +     4176.04134244900*t);
+   emb_y_0+=     0.00000000329 * cos( 5.87961134010 +        3.93215326310*t);
+   emb_y_0+=     0.00000000374 * cos( 5.15697532254 +     5642.19824260920*t);
+   emb_y_0+=     0.00000000320 * cos( 0.92923666608 +    10818.13528691580*t);
+   emb_y_0+=     0.00000000365 * cos( 5.96664322939 +     6660.44945790720*t);
+   emb_y_0+=     0.00000000339 * cos( 4.53749520696 +     6923.95345737360*t);
+   emb_y_0+=     0.00000000357 * cos( 1.12153755723 +    16460.33352952499*t);
+   emb_y_0+=     0.00000000386 * cos( 0.71941423743 +     2636.72547263700*t);
+   emb_y_0+=     0.00000000273 * cos( 5.79109432105 +    26709.64694241340*t);
+   emb_y_0+=     0.00000000265 * cos( 3.53930695872 +    19378.91851506880*t);
+   emb_y_0+=     0.00000000325 * cos( 0.92198442483 +    12359.96615154560*t);
+   emb_y_0+=     0.00000000278 * cos( 3.86335874809 +    20597.24396304120*t);
+   emb_y_0+=     0.00000000294 * cos( 2.59935512947 +    28286.99048486120*t);
+   emb_y_0+=     0.00000000269 * cos( 4.91776226798 +     6816.28993343500*t);
+   emb_y_0+=     0.00000000357 * cos( 2.74654149425 +    30666.15495843280*t);
+   emb_y_0+=     0.00000000344 * cos( 1.94480070409 +     6147.45052498140*t);
+   emb_y_0+=     0.00000000255 * cos( 3.37524772579 +     1375.77379984580*t);
+   emb_y_0+=     0.00000000308 * cos( 5.40069482744 +    29864.33402730900*t);
+   emb_y_0+=     0.00000000334 * cos( 1.56223513612 +     6418.70117500140*t);
+   emb_y_0+=     0.00000000316 * cos( 3.97125662279 +     6055.59971720240*t);
+   emb_y_0+=     0.00000000243 * cos( 2.53093250846 +     6310.47740608820*t);
+   emb_y_0+=     0.00000000243 * cos( 0.97600841060 +     6255.67429389460*t);
+   emb_y_0+=     0.00000000296 * cos( 4.79138384974 +     6129.29703950660*t);
+   emb_y_0+=     0.00000000301 * cos( 5.81886973641 +     6510.55198278040*t);
+   emb_y_0+=     0.00000000329 * cos( 5.86928788098 +      377.37360791580*t);
+   emb_y_0+=     0.00000000282 * cos( 3.64505853337 +      263.08392337280*t);
+   emb_y_0+=     0.00000000267 * cos( 1.84936252726 +    24705.70520908959*t);
+   emb_y_0+=     0.00000000285 * cos( 0.93414827012 +     5469.52556603160*t);
+   emb_y_0+=     0.00000000234 * cos( 0.33992064306 +     5017.50837136500*t);
+   emb_y_0+=     0.00000000286 * cos( 4.99886287753 +     6436.85466047620*t);
+   emb_y_0+=     0.00000000309 * cos( 0.57495592500 +    27511.46787353720*t);
+   emb_y_0+=     0.00000000263 * cos( 6.04953863141 +    11712.95531823080*t);
+   emb_y_0+=     0.00000000285 * cos( 5.85104532566 +    16522.65971600220*t);
+   emb_y_0+=     0.00000000230 * cos( 5.72898960063 +     1265.56747862640*t);
+   emb_y_0+=     0.00000000235 * cos( 4.94143074133 +     4487.81740627040*t);
+   emb_y_0+=     0.00000000253 * cos( 2.08490293021 +     4377.61108505100*t);
+   emb_y_0+=     0.00000000312 * cos( 1.61150992635 +    11720.06886523160*t);
+   emb_y_0+=     0.00000000220 * cos( 6.00356459702 +    20199.09495963300*t);
+   emb_y_0+=     0.00000000214 * cos( 0.52396088270 +     2221.85663459700*t);
+   emb_y_0+=     0.00000000213 * cos( 2.22365237019 +    30355.99731976780*t);
+   emb_y_0+=     0.00000000259 * cos( 1.86047749915 +     6282.83203250789*t);
+   emb_y_0+=     0.00000000259 * cos( 1.64646341991 +     6283.31966747490*t);
+   emb_y_0+=     0.00000000243 * cos( 2.57267015777 +     7096.62613395120*t);
+   emb_y_0+=     0.00000000221 * cos( 4.65757923242 +     5540.08578945880*t);
+   emb_y_0+=     0.00000000246 * cos( 4.26495076165 +     8671.96987044060*t);
+   emb_y_0+=     0.00000000208 * cos( 3.10897679499 +    10440.27429260360*t);
+   emb_y_0+=     0.00000000199 * cos( 5.64807046697 +    11216.28429032400*t);
+   emb_y_0+=     0.00000000209 * cos( 3.71139099761 +      419.48464387520*t);
+   emb_y_0+=     0.00000000207 * cos( 1.49549061720 +    12562.56127133100*t);
+   emb_y_0+=     0.00000000212 * cos( 0.18391974040 +     7129.15868474260*t);
+   emb_y_0+=     0.00000000224 * cos( 5.58680388524 +     5636.06501667660*t);
+   emb_y_0+=     0.00000000192 * cos( 0.47333599881 +    17271.88400752640*t);
+   emb_y_0+=     0.00000000241 * cos( 3.52161359834 +    15906.76412668260*t);
+   emb_y_0+=     0.00000000257 * cos( 0.87147328095 +    48739.85989708300*t);
+   emb_y_0+=     0.00000000184 * cos( 3.83487369725 +     9380.95967271720*t);
+   emb_y_0+=     0.00000000256 * cos( 0.73074058890 +      224.34479570190*t);
+   emb_y_0+=     0.00000000250 * cos( 2.67103537581 +     5934.15142954340*t);
+   emb_y_0+=     0.00000000185 * cos( 1.51386870105 +    19800.94595622480*t);
+   emb_y_0+=     0.00000000189 * cos( 2.52919302388 +    16062.18452611680*t);
+   emb_y_0+=     0.00000000175 * cos( 3.84093628240 +    11614.43329373220*t);
+   emb_y_0+=     0.00000000175 * cos( 6.15313646522 +     3744.82734573780*t);
+   emb_y_0+=     0.00000000240 * cos( 0.45238664563 +        1.48447270830*t);
+   emb_y_0+=     0.00000000241 * cos( 5.00040455119 +     8390.11035753380*t);
+   emb_y_0+=     0.00000000193 * cos( 4.71506547135 +      394.62588505920*t);
+   emb_y_0+=     0.00000000169 * cos( 0.80098424793 +      742.99006053260*t);
+   emb_y_0+=     0.00000000189 * cos( 4.03411872458 +      401.67212175720*t);
+   emb_y_0+=     0.00000000230 * cos( 0.83612842798 +     6632.00027043940*t);
+   emb_y_0+=     0.00000000227 * cos( 4.16811397278 +    34596.36465465240*t);
+   emb_y_0+=     0.00000000165 * cos( 0.39341190294 +    10454.50138660520*t);
+   emb_y_0+=     0.00000000221 * cos( 2.07224292531 +     6206.80977871580*t);
+   emb_y_0+=     0.00000000201 * cos( 6.23059857608 +     3097.88382272579*t);
+   emb_y_0+=     0.00000000205 * cos( 2.56533473350 +    16627.37091537720*t);
+   emb_y_0+=     0.00000000178 * cos( 0.01807925983 +     4797.09572892620*t);
+   emb_y_0+=     0.00000000225 * cos( 3.43883654828 +      266.60704172180*t);
+   emb_y_0+=     0.00000000167 * cos( 2.40308228314 +     5714.25397596400*t);
+   emb_y_0+=     0.00000000167 * cos( 1.16660958168 +    10344.29506538580*t);
+   emb_y_0+=     0.00000000172 * cos( 3.81893679688 +     6332.83287546320*t);
+   emb_y_0+=     0.00000000172 * cos( 5.97118942936 +     6233.31882451960*t);
+   emb_y_0+=     0.00000000214 * cos( 2.77237801335 +    24492.40611365159*t);
+   emb_y_0+=     0.00000000159 * cos( 2.33166453392 +    13198.93543929600*t);
+   emb_y_0+=     0.00000000207 * cos( 1.92518585993 +    12569.74212863460*t);
+   emb_y_0+=     0.00000000176 * cos( 2.48119329994 +    12012.58229714040*t);
+   emb_y_0+=     0.00000000196 * cos( 5.41851032427 +    11087.28512591840*t);
+   emb_y_0+=     0.00000000161 * cos( 2.18455666842 +    12132.43996210600*t);
+   emb_y_0+=     0.00000000170 * cos( 5.35052715864 +     6272.03014972750*t);
+   emb_y_0+=     0.00000000170 * cos( 4.43959906760 +     6294.12155025530*t);
+   emb_y_0+=     0.00000000168 * cos( 3.97617357428 +    17686.75284556640*t);
+   emb_y_0+=     0.00000000199 * cos( 5.15620822027 +    33326.57873317420*t);
+   emb_y_0+=     0.00000000194 * cos( 2.13712729333 +       28.44918746780*t);
+   emb_y_0+=     0.00000000144 * cos( 2.03929994576 +    16737.57723659660*t);
+   emb_y_0+=     0.00000000148 * cos( 1.10385881777 +     6851.89772401880*t);
+   emb_y_0+=     0.00000000146 * cos( 4.95201849548 +    17370.36097590980*t);
+   emb_y_0+=     0.00000000143 * cos( 5.56441147227 +     5618.31980486140*t);
+   emb_y_0+=     0.00000000140 * cos( 5.93460328071 +     2008.55753915900*t);
+   emb_y_0+=     0.00000000162 * cos( 3.21674523006 +    12551.92460598120*t);
+   emb_y_0+=     0.00000000143 * cos( 3.89365374642 +     5010.39482436420*t);
+   emb_y_0+=     0.00000000161 * cos( 1.75025798421 +    31441.67756975680*t);
+   emb_y_0+=     0.00000000151 * cos( 4.72729575304 +    16723.35014259500*t);
+   emb_y_0+=     0.00000000135 * cos( 4.79054860554 +    12431.08661994740*t);
+   emb_y_0+=     0.00000000148 * cos( 3.93657905776 +    28766.92442448400*t);
+   emb_y_0+=     0.00000000146 * cos( 6.27087503385 +    12701.21678001820*t);
+   emb_y_0+=     0.00000000137 * cos( 1.02989768135 +       11.04570026390*t);
+   emb_y_0+=     0.00000000167 * cos( 5.98371371124 +    17473.45375012840*t);
+   emb_y_0+=     0.00000000161 * cos( 1.02898808001 +    20452.86941222180*t);
+   emb_y_0+=     0.00000000124 * cos( 4.22571513469 +     6947.83189512140*t);
+   emb_y_0+=     0.00000000135 * cos( 4.95838914470 +     3531.52825029980*t);
+   emb_y_0+=     0.00000000122 * cos( 4.12184077581 +     3583.34103067380*t);
+   emb_y_0+=     0.00000000155 * cos( 5.62826239458 +     6058.73105428950*t);
+   emb_y_0+=     0.00000000120 * cos( 5.90379805618 +    12808.88030395680*t);
+   emb_y_0+=     0.00000000145 * cos( 3.70967890118 +       95.97922721780*t);
+   emb_y_0+=     0.00000000120 * cos( 4.26997270152 +    12528.12402734700*t);
+   emb_y_0+=     0.00000000119 * cos( 5.90934719297 +    19247.37655338240*t);
+   emb_y_0+=     0.00000000117 * cos( 5.45034772919 +    11610.55195837420*t);
+   emb_y_0+=     0.00000000163 * cos( 1.54432161045 +      661.23292678100*t);
+   emb_y_0+=     0.00000000161 * cos( 2.28994816795 +    52670.06959330260*t);
+   emb_y_0+=     0.00000000163 * cos( 2.78769644779 +     9929.42622734580*t);
+   emb_y_0+=     0.00000000119 * cos( 2.71628583938 +    12249.75983032620*t);
+   emb_y_0+=     0.00000000145 * cos( 1.17170397718 +     6360.74962041940*t);
+   emb_y_0+=     0.00000000145 * cos( 2.33523694188 +     6205.40207956340*t);
+   emb_y_0+=     0.00000000156 * cos( 1.00708912320 +     4307.58330413540*t);
+   emb_y_0+=     0.00000000118 * cos( 5.95164477875 +     6302.74561089119*t);
+   emb_y_0+=     0.00000000118 * cos( 3.83848144749 +     6263.40608909160*t);
+   emb_y_0+=     0.00000000135 * cos( 2.80387650268 +    19804.82729158280*t);
+   emb_y_0+=     0.00000000112 * cos( 3.88136222265 +    16840.67001081519*t);
+   emb_y_0+=     0.00000000148 * cos( 3.55909740786 +    28237.23345938940*t);
+   emb_y_0+=     0.00000000124 * cos( 5.50657775921 +    10239.58386601080*t);
+   emb_y_0+=     0.00000000137 * cos( 2.53358646423 +     6250.54329920000*t);
+   emb_y_0+=     0.00000000137 * cos( 0.97335445483 +     6315.60840078280*t);
+   emb_y_0+=     0.00000000111 * cos( 0.70362725529 +     6339.97422492700*t);
+   emb_y_0+=     0.00000000111 * cos( 2.80331366377 +     6226.17747505580*t);
+   emb_y_0+=     0.00000000139 * cos( 4.20829691432 +    11609.86254401220*t);
+   emb_y_0+=     0.00000000124 * cos( 3.12647421017 +    19402.79695281660*t);
+   emb_y_0+=     0.00000000123 * cos( 0.80538262187 +    24734.15439655740*t);
+   emb_y_0+=     0.00000000107 * cos( 0.59542543867 +    12463.05892576420*t);
+   emb_y_0+=     0.00000000140 * cos( 5.62756230388 +    38526.57435087200*t);
+   emb_y_0+=     0.00000000123 * cos( 5.71833316259 +     6305.85105144220*t);
+   emb_y_0+=     0.00000000123 * cos( 4.07179306365 +     6260.30064854060*t);
+   emb_y_0+=     0.00000000113 * cos( 0.38311301542 +     2111.65031337760*t);
+   emb_y_0+=     0.00000000117 * cos( 1.11114131260 +    24336.00539314920*t);
+   emb_y_0+=     0.00000000131 * cos( 2.04122395494 +    18415.51581209739*t);
+   emb_y_0+=     0.00000000106 * cos( 2.18835416515 +    12341.80690428090*t);
+   emb_y_0+=     0.00000000110 * cos( 3.31587494701 +     6276.44729109139*t);
+   emb_y_0+=     0.00000000110 * cos( 0.19106597205 +     6289.70440889141*t);
+   emb_y_0+=     0.00000000103 * cos( 1.18729919186 +     6324.62664097620*t);
+   emb_y_0+=     0.00000000103 * cos( 2.31964172720 +     6241.52505900660*t);
+   emb_y_0+=     0.00000000131 * cos( 5.93782050272 +     5791.41255753260*t);
+   emb_y_0+=     0.00000000106 * cos( 1.70601019340 +    12772.33724842000*t);
+   emb_y_0+=     0.00000000143 * cos( 6.22488817389 +     2069.00683490660*t);
+   emb_y_0+=     0.00000000143 * cos( 5.63221604750 +    18772.96147869860*t);
+   emb_y_0+=     0.00000000115 * cos( 5.35511207484 +     3641.73457151920*t);
+   emb_y_0+=     0.00000000121 * cos( 3.53431033420 +    17996.03116822220*t);
+   emb_y_0+=     0.00000000126 * cos( 5.90533164276 +    18822.92923017440*t);
+   emb_y_0+=     0.00000000115 * cos( 3.09464836385 +    12570.08385324590*t);
+   emb_y_0+=     0.00000000119 * cos( 0.26217017763 +     9814.60410029120*t);
+   emb_y_0+=     0.00000000138 * cos( 0.51855418741 +     6356.37297585040*t);
+   emb_y_0+=     0.00000000138 * cos( 2.98838673165 +     6209.77872413240*t);
+   emb_y_0+=     0.00000000100 * cos( 3.07908951574 +        5.52292430740*t);
+   emb_y_0+=     0.00000000135 * cos( 5.28292512146 +    18699.66435283960*t);
+   emb_y_0+=     0.00000000104 * cos( 0.57981411648 +    12560.62877567540*t);
+   emb_y_0+=     0.00000000098 * cos( 4.00434933508 +     5815.11085963699*t);
+   emb_y_0+=     0.00000000096 * cos( 3.47793123567 +    22743.40937951640*t);
+   emb_y_0+=     0.00000000099 * cos( 2.68455377979 +    78263.70942472259*t);
+   emb_y_0+=     0.00000000114 * cos( 2.18292460802 +     6155.05700665400*t);
+   emb_y_0+=     0.00000000114 * cos( 1.32401631104 +     6411.09469332880*t);
+   emb_y_0+=     0.00000000093 * cos( 0.34809743975 +    11190.37790013700*t);
+   emb_y_0+=     0.00000000095 * cos( 4.06228552408 +    11300.58422135640*t);
+   emb_y_0+=     0.00000000117 * cos( 2.41831042685 +    24383.07910844140*t);
+   emb_y_0+=     0.00000000110 * cos( 1.22040625314 +     6312.50535852740*t);
+   emb_y_0+=     0.00000000110 * cos( 2.28653466592 +     6253.64634145540*t);
+   emb_y_0+=     0.00000000091 * cos( 1.42377843647 +    12829.23562335560*t);
+   emb_y_0+=     0.00000000120 * cos( 4.42365311529 +    17892.93839400359*t);
+   emb_y_0+=     0.00000000090 * cos( 4.05016333707 +    15664.03552270859*t);
+   emb_y_0+=     0.00000000126 * cos( 0.70432668892 +     1293.48422358260*t);
+   emb_y_0+=     0.00000000120 * cos( 0.51717377640 +    14158.74771361560*t);
+   emb_y_0+=     0.00000000110 * cos( 1.99803947181 +    16097.67995028260*t);
+   emb_y_0+=     0.00000000092 * cos( 1.75886109740 +    21424.46664430340*t);
+   emb_y_0+=     0.00000000099 * cos( 2.04878528903 +     6432.63904712600*t);
+   emb_y_0+=     0.00000000088 * cos( 0.12768507477 +    15141.39079431200*t);
+   emb_y_0+=     0.00000000089 * cos( 5.78577662098 +     6751.04084034580*t);
+   emb_y_0+=     0.00000000096 * cos( 1.07414401415 +     7349.57132718140*t);
+   emb_y_0+=     0.00000000094 * cos( 1.30652024350 +     2699.73481931760*t);
+   emb_y_0+=     0.00000000087 * cos( 1.63535344372 +    23141.55838292460*t);
+   emb_y_0+=     0.00000000103 * cos( 2.80666225492 +    18202.21671665939*t);
+   emb_y_0+=     0.00000000099 * cos( 5.29408634455 +     6313.98997562640*t);
+   emb_y_0+=     0.00000000099 * cos( 4.49603988169 +     6252.16172435640*t);
+   emb_y_0+=     0.00000000113 * cos( 3.06789441452 +    10557.59416082380*t);
+   emb_y_0+=     0.00000000098 * cos( 1.66301652528 +     6126.67512948900*t);
+   emb_y_0+=     0.00000000098 * cos( 1.84392439378 +     6439.47657049380*t);
+   emb_y_0+=     0.00000000086 * cos( 6.05385675766 +    23123.74586080659*t);
+   emb_y_0+=     0.00000000094 * cos( 5.54106526599 +    23539.70738633280*t);
+   emb_y_0+=     0.00000000083 * cos( 4.71523084157 +     8858.31494432060*t);
+   emb_y_0+=     0.00000000101 * cos( 4.94270501848 +    18429.74290609900*t);
+   emb_y_0+=     0.00000000088 * cos( 2.19663281211 +    17053.96910625319*t);
+   emb_y_0+=     0.00000000113 * cos( 0.47970289838 +    12985.63634385800*t);
+   emb_y_0+=     0.00000000082 * cos( 4.90614256431 +     5483.25472482600*t);
+   emb_y_0+=     0.00000000095 * cos( 0.81797133475 +    11925.27409260060*t);
+   emb_y_0+=     0.00000000084 * cos( 2.66665772953 +     6379.05507720920*t);
+   emb_y_0+=     0.00000000084 * cos( 0.84028318953 +     6187.09662277360*t);
+   emb_y_0+=     0.00000000093 * cos( 3.10680657507 +    35371.88726597640*t);
+   emb_y_0+=     0.00000000105 * cos( 1.20399152892 +     6434.97313107660*t);
+   emb_y_0+=     0.00000000105 * cos( 2.30294939014 +     6131.17856890620*t);
+   emb_y_0+=     0.00000000114 * cos( 2.50001433093 +     8258.56839584740*t);
+   emb_y_0+=     0.00000000084 * cos( 3.82196413310 +    32217.20018108080*t);
+   emb_y_0+=     0.00000000082 * cos( 5.40278723857 +     3958.12644117580*t);
+   emb_y_0+=     0.00000000102 * cos( 1.32703960779 +     5767.61197889840*t);
+   emb_y_0+=     0.00000000102 * cos( 2.17990131127 +     6798.53972108440*t);
+   emb_y_0+=     0.00000000092 * cos( 6.07396273786 +    18216.44381066100*t);
+   emb_y_0+=     0.00000000112 * cos( 3.91825505548 +    12188.77809206700*t);
+   emb_y_0+=     0.00000000086 * cos( 4.18439486840 +    16193.65917750039*t);
+   emb_y_0+=     0.00000000080 * cos( 2.01988251452 +    24080.03501677719*t);
+   emb_y_0+=     0.00000000093 * cos( 4.66163367881 +    23937.85638974100*t);
+   emb_y_0+=     0.00000000088 * cos( 4.57496913724 +     1573.82042409880*t);
+   emb_y_0+=     0.00000000087 * cos( 3.89286413327 +     1580.86666079680*t);
+   emb_y_0+=     0.00000000080 * cos( 1.70826898567 +     6485.32924516550*t);
+   emb_y_0+=     0.00000000080 * cos( 1.79867193339 +     6080.82245481730*t);
+   emb_y_0+=     0.00000000079 * cos( 6.16114738580 +    40796.33892267419*t);
+   emb_y_0+=     0.00000000079 * cos( 2.96735408932 +    36109.38220466459*t);
+   emb_y_0+=     0.00000000088 * cos( 5.58306387433 +    16943.76278503380*t);
+   emb_y_0+=     0.00000000093 * cos( 0.22422436515 +      283.85931886520*t);
+   emb_y_0+=     0.00000000088 * cos( 1.67821631800 +    12537.70251251500*t);
+   emb_y_0+=     0.00000000079 * cos( 4.76362160922 +    41194.48792608240*t);
+   emb_y_0+=     0.00000000081 * cos( 5.19149105668 +     3956.50801601940*t);
+   emb_y_0+=     0.00000000103 * cos( 3.71008010926 +    56600.27928952220*t);
+   emb_y_0+=     0.00000000079 * cos( 2.93697301370 +     5621.84292321040*t);
+   emb_y_0+=     0.00000000096 * cos( 0.41418225890 +     5888.44996493220*t);
+   emb_y_0+=     0.00000000096 * cos( 3.09275866016 +     6677.70173505060*t);
+   emb_y_0+=     0.00000000074 * cos( 4.41044639229 +     6304.41149045840*t);
+   emb_y_0+=     0.00000000074 * cos( 5.37967983395 +     6261.74020952440*t);
+   emb_y_0+=     0.00000000073 * cos( 1.05049706961 +     6695.44694686580*t);
+   emb_y_0+=     0.00000000073 * cos( 2.45644384945 +     5870.70475311700*t);
+   emb_y_0+=     0.00000000073 * cos( 1.61765616896 +    16207.88627150200*t);
+   emb_y_0+=     0.00000000085 * cos( 2.08161902252 +    12850.01101884800*t);
+   emb_y_0+=     0.00000000072 * cos( 5.03595593677 +     6312.89728814020*t);
+   emb_y_0+=     0.00000000072 * cos( 4.75417028947 +     6253.25441184259*t);
+   emb_y_0+=     0.00000000100 * cos( 3.80714752393 +     5316.10497255580*t);
+   emb_y_0+=     0.00000000085 * cos( 1.33904345426 +    12282.29238111760*t);
+   emb_y_0+=     0.00000000077 * cos( 5.07586430304 +    22345.26037610820*t);
+   emb_y_0+=     0.00000000076 * cos( 0.41146256256 +    63658.87775083760*t);
+   emb_y_0+=     0.00000000087 * cos( 2.06336091493 +    18106.23748944160*t);
+   emb_y_0+=     0.00000000070 * cos( 0.44640955505 +     6442.01936777460*t);
+   emb_y_0+=     0.00000000070 * cos( 3.06053136401 +     6124.13233220820*t);
+   emb_y_0+=     0.00000000074 * cos( 5.61945589711 +    15265.88651930040*t);
+   emb_y_0+=     0.00000000078 * cos( 3.38283383527 +    23550.34405168259*t);
+   emb_y_0+=     0.00000000077 * cos( 5.00656029204 +     5237.92101380380*t);
+   emb_y_0+=     0.00000000075 * cos( 3.64547243832 +    27832.03821928320*t);
+   emb_y_0+=     0.00000000082 * cos( 3.91623659935 +     6774.73914245020*t);
+   emb_y_0+=     0.00000000067 * cos( 4.88284279928 +     7082.89697515680*t);
+   emb_y_0+=     0.00000000067 * cos( 3.24232940659 +     6171.88920770380*t);
+   emb_y_0+=     0.00000000067 * cos( 0.26461151247 +     6394.26249227900*t);
+   emb_y_0+=     0.00000000080 * cos( 3.79677325613 +    22380.75580027400*t);
+   emb_y_0+=     0.00000000070 * cos( 0.47599032576 +     6175.41232605280*t);
+   emb_y_0+=     0.00000000070 * cos( 3.03095059330 +     6390.73937393000*t);
+   emb_y_0+=     0.00000000077 * cos( 0.29137018417 +     9924.81042151060*t);
+   emb_y_0+=     0.00000000091 * cos( 0.90446873976 +     6652.77566593180*t);
+   emb_y_0+=     0.00000000091 * cos( 2.60247217930 +     5913.37603405100*t);
+   emb_y_0+=     0.00000000067 * cos( 4.84479908786 +     6300.88837210940*t);
+   emb_y_0+=     0.00000000067 * cos( 4.94532713838 +     6265.26332787340*t);
+   emb_y_0+=     0.00000000088 * cos( 3.59122542798 +    10660.68693504240*t);
+   emb_y_0+=     0.00000000066 * cos( 3.01430181986 +     5959.57043333400*t);
+   emb_y_0+=     0.00000000080 * cos( 5.49178527606 +    23646.32327890039*t);
+   emb_y_0+=     0.00000000068 * cos( 2.07232411643 +     4897.18057365520*t);
+   emb_y_0+=     0.00000000069 * cos( 0.56996801147 +     6944.30877677240*t);
+   emb_y_0+=     0.00000000073 * cos( 2.96334792367 +     6195.76764545159*t);
+   emb_y_0+=     0.00000000073 * cos( 0.54359299539 +     6370.38405453120*t);
+   emb_y_0+=     0.00000000071 * cos( 3.57166978930 +    17363.24742890899*t);
+   emb_y_0+=     0.00000000074 * cos( 1.56614709067 +    27043.50288318280*t);
+   emb_y_0+=     0.00000000066 * cos( 2.38194509538 +     6251.10354417460*t);
+   emb_y_0+=     0.00000000066 * cos( 1.12499582368 +     6315.04815580820*t);
+   emb_y_0+=     0.00000000064 * cos( 3.50738755186 +      323.50541665740*t);
+   emb_y_0+=     0.00000000078 * cos( 0.83740696355 +    42456.78404709160*t);
+   emb_y_0+=     0.00000000063 * cos( 2.70916268945 +    11823.16163945020*t);
+   emb_y_0+=     0.00000000064 * cos( 4.80330562067 +    28774.03797148480*t);
+   emb_y_0+=     0.00000000079 * cos( 5.66337403718 +    24065.80792277559*t);
+   emb_y_0+=     0.00000000072 * cos( 2.11096843463 +    18100.00325845000*t);
+   emb_y_0+=     0.00000000066 * cos( 1.21053452412 +    17576.54652434700*t);
+   emb_y_0+=     0.00000000068 * cos( 0.52464381855 +    18003.14471522300*t);
+   emb_y_0+=     0.00000000068 * cos( 1.33733194362 +    28759.81087748319*t);
+   emb_y_0+=     0.00000000062 * cos( 2.25046717820 +     3855.03366695720*t);
+   emb_y_0+=     0.00000000060 * cos( 5.83008017266 +    18606.49894600020*t);
+   emb_y_0+=     0.00000000079 * cos( 2.49476757272 +    22476.73502749179*t);
+   emb_y_0+=     0.00000000073 * cos( 4.61251602377 +    23440.13773046320*t);
+   emb_y_0+=     0.00000000060 * cos( 1.65871121739 +    13088.72911807660*t);
+   emb_y_0+=     0.00000000078 * cos( 2.99014426818 +    12564.66722727450*t);
+   emb_y_0+=     0.00000000060 * cos( 2.90135289748 +     7548.64332861780*t);
+   emb_y_0+=     0.00000000073 * cos( 4.73963217435 +      735.87651353180*t);
+   emb_y_0+=     0.00000000059 * cos( 2.43330626408 +     5113.48759858280*t);
+   emb_y_0+=     0.00000000075 * cos( 2.85043788676 +    18624.88275427230*t);
+   emb_y_0+=     0.00000000059 * cos( 5.13628041810 +    41990.78593289880*t);
+   emb_y_0+=     0.00000000081 * cos( 5.98297918907 +     7250.04672742700*t);
+   emb_y_0+=     0.00000000059 * cos( 2.39193379192 +     6289.14176562120*t);
+   emb_y_0+=     0.00000000059 * cos( 1.11500712714 +     6277.00993436160*t);
+   emb_y_0+=     0.00000000072 * cos( 4.87534601803 +     8982.81066930900*t);
+   emb_y_0+=     0.00000000057 * cos( 3.30911091342 +      792.77488846740*t);
+   emb_y_0+=     0.00000000057 * cos( 1.61290511224 +    10027.90319572920*t);
+   emb_y_0+=     0.00000000072 * cos( 5.73653948228 +     4598.02372748980*t);
+   emb_y_0+=     0.00000000060 * cos( 0.63449806943 +      647.01083331480*t);
+   emb_y_0+=     0.00000000066 * cos( 0.80721495515 +     5220.17080145320*t);
+   emb_y_0+=     0.00000000056 * cos( 5.24002840662 +    35050.00027447539*t);
+   emb_y_0+=     0.00000000077 * cos( 2.14165124499 +    46386.99374331120*t);
+   emb_y_0+=     0.00000000068 * cos( 1.14499338723 +     3166.41643773160*t);
+   emb_y_0+=     0.00000000056 * cos( 0.39917748570 +     6339.87911216120*t);
+   emb_y_0+=     0.00000000056 * cos( 3.10776343336 +     6226.27258782160*t);
+   emb_y_0+=     0.00000000056 * cos( 0.95649088929 +    12303.06777661000*t);
+   emb_y_0+=     0.00000000057 * cos( 1.98721525004 +     6702.00024889200*t);
+   emb_y_0+=     0.00000000057 * cos( 1.51972566902 +     5864.15145109080*t);
+   emb_y_0+=     0.00000000069 * cos( 0.33788710387 +    23536.11695768099*t);
+   emb_y_0+=     0.00000000059 * cos( 1.81421782103 +    67589.08744705719*t);
+   emb_y_0+=     0.00000000076 * cos( 0.77801824814 +    23226.83863502520*t);
+   emb_y_0+=     0.00000000059 * cos( 1.03968301691 +       15.25247118500*t);
+   emb_y_0+=     0.00000000073 * cos( 4.43868056110 +    30348.88377276700*t);
+   emb_y_0+=     0.00000000054 * cos( 4.96201305869 +    36147.40987730040*t);
+   emb_y_0+=     0.00000000070 * cos( 6.27818224660 +    12299.54465826100*t);
+   emb_y_0+=     0.00000000060 * cos( 5.08461088458 +    12171.52581492360*t);
+   emb_y_0+=     0.00000000074 * cos( 1.19928905734 +     3684.37804999020*t);
+   emb_y_0+=     0.00000000058 * cos( 1.58539308271 +     6716.78758786820*t);
+   emb_y_0+=     0.00000000061 * cos( 3.25163800829 +    27707.54249429480*t);
+   emb_y_0+=     0.00000000059 * cos( 5.76308524321 +    12164.47957822560*t);
+   emb_y_0+=     0.00000000070 * cos( 0.07217173698 +    16833.55646381439*t);
+   emb_y_0+=     0.00000000063 * cos( 2.16224668422 +     5227.62607306530*t);
+   emb_y_0+=     0.00000000063 * cos( 5.23287718161 +     5219.76176653910*t);
+   emb_y_0+=     0.00000000056 * cos( 2.45377587379 +     3178.14579056760*t);
+   emb_y_0+=     0.00000000072 * cos( 2.62702736253 +     4583.79663348820*t);
+   emb_y_0+=     0.00000000052 * cos( 0.96546586882 +     3627.50747751760*t);
+   emb_y_0+=     0.00000000052 * cos( 0.20415505351 +     6155.12431695680*t);
+   emb_y_0+=     0.00000000052 * cos( 3.30278586555 +     6411.02738302600*t);
+   emb_y_0+=     0.00000000060 * cos( 0.48342399631 +     5244.04923920100*t);
+   emb_y_0+=     0.00000000058 * cos( 2.72226132172 +    30775.48196364300*t);
+   emb_y_0+=     0.00000000052 * cos( 5.80291300969 +    30639.85663863300*t);
+   emb_y_0+=     0.00000000053 * cos( 5.75258856716 +    37724.75341974820*t);
+   emb_y_0+=     0.00000000062 * cos( 2.89766241568 +     6908.74604230380*t);
+   emb_y_0+=     0.00000000062 * cos( 0.60927850338 +     5657.40565767900*t);
+   emb_y_0+=     0.00000000050 * cos( 3.11011154958 +     2575.23909432920*t);
+   emb_y_0+=     0.00000000054 * cos( 1.15676928987 +    18208.34994259200*t);
+   emb_y_0+=     0.00000000050 * cos( 3.12376167242 +     5881.40372823420*t);
+   emb_y_0+=     0.00000000050 * cos( 0.38317924664 +     6684.74797174860*t);
+   emb_y_0+=     0.00000000068 * cos( 4.61950794623 +    96677.89886304220*t);
+   emb_y_0+=     0.00000000054 * cos( 0.77163765121 +    41654.96311596780*t);
+   emb_y_0+=     0.00000000066 * cos( 5.13003693613 +    60530.48898574180*t);
+   emb_y_0+=     0.00000000050 * cos( 1.64593222191 +    43738.80234596580*t);
+   emb_y_0+=     0.00000000048 * cos( 5.44674337925 +    25685.87280280800*t);
+   emb_y_0+=     0.00000000059 * cos( 2.71376008731 +     3981.49003408200*t);
+   emb_y_0+=     0.00000000058 * cos( 0.74658643142 +    14169.79356223040*t);
+   emb_y_0+=     0.00000000051 * cos( 1.43461659855 +     7668.97112632760*t);
+   emb_y_0+=     0.00000000048 * cos( 3.23576011756 +     3646.35037735440*t);
+   emb_y_0+=     0.00000000054 * cos( 4.50803792908 +    39302.09696219600*t);
+   emb_y_0+=     0.00000000057 * cos( 2.62704971421 +      799.82112516540*t);
+   emb_y_0+=     0.00000000061 * cos( 3.28745263972 +    44034.12758953940*t);
+   emb_y_0+=     0.00000000050 * cos( 3.60626874009 +    17583.66007134780*t);
+   emb_y_0+=     0.00000000051 * cos( 1.16028405425 +    38650.17350619900*t);
+   emb_y_0+=     0.00000000053 * cos( 1.90122336672 +    11080.17157891760*t);
+   emb_y_0+=     0.00000000059 * cos( 0.68886289824 +    62883.35513951360*t);
+   emb_y_0+=     0.00000000046 * cos( 0.83996053249 +    25287.72379939980*t);
+   emb_y_0+=     0.00000000049 * cos( 3.59486880927 +     5490.30096152400*t);
+   emb_y_0+=     0.00000000058 * cos( 3.83698375555 +     6076.37511269480*t);
+   emb_y_0+=     0.00000000058 * cos( 5.95314247069 +     6489.77658728800*t);
+   emb_y_0+=     0.00000000047 * cos( 4.43378504018 +      170.67287061920*t);
+   emb_y_0+=     0.00000000045 * cos( 5.81363953764 +    33318.81606245820*t);
+   emb_y_0+=     0.00000000046 * cos( 3.19651245272 +     4480.70385926960*t);
+   emb_y_0+=     0.00000000060 * cos( 3.92066214096 +    28313.28880466100*t);
+   emb_y_0+=     0.00000000061 * cos( 3.53510681440 +      491.66329245880*t);
+   emb_y_0+=     0.00000000051 * cos( 6.00904914233 +     3428.43547608120*t);
+   emb_y_0+=     0.00000000060 * cos( 3.51283911810 +    50317.20343953080*t);
+   emb_y_0+=     0.00000000054 * cos( 2.28043683109 +     6559.82162185580*t);
+   emb_y_0+=     0.00000000054 * cos( 1.22650408797 +     6006.33007812700*t);
+   emb_y_0+=     0.00000000046 * cos( 3.84965479184 +     4384.72463205180*t);
+   emb_y_0+=     0.00000000045 * cos( 0.29195078507 +    44136.95134937400*t);
+   emb_y_0+=     0.00000000044 * cos( 2.71274650311 +     2301.58581590939*t);
+   emb_y_0+=     0.00000000048 * cos( 4.64164514773 +    35579.69123957000*t);
+   emb_y_0+=     0.00000000051 * cos( 4.91887651070 +    18326.65013188039*t);
+   emb_y_0+=     0.00000000046 * cos( 3.15339406184 +    16703.06213349900*t);
+   emb_y_0+=     0.00000000048 * cos( 2.39477993061 +     5209.46682580060*t);
+   emb_y_0+=     0.00000000046 * cos( 2.47672465258 +    12555.10599971890*t);
+   emb_y_0+=     0.00000000052 * cos( 5.14565632171 +     6574.04871585740*t);
+   emb_y_0+=     0.00000000052 * cos( 4.64446990453 +     5992.10298412540*t);
+   emb_y_0+=     0.00000000044 * cos( 3.80836113401 +    69941.95360082899*t);
+   emb_y_0+=     0.00000000050 * cos( 3.54529295551 +    22030.21295466960*t);
+   emb_y_0+=     0.00000000043 * cos( 1.20775471242 +    20894.77567920879*t);
+   emb_y_0+=     0.00000000044 * cos( 2.61939991177 +      533.21408344360*t);
+   emb_y_0+=     0.00000000045 * cos( 4.65113002963 +    46360.69542351140*t);
+   emb_y_0+=     0.00000000047 * cos( 2.17023181784 +    45585.17281218740*t);
+   emb_y_0+=     0.00000000043 * cos( 1.54528183996 +     8401.83971036980*t);
+   emb_y_0+=     0.00000000054 * cos( 3.62371389676 +    24279.10701821359*t);
+   emb_y_0+=     0.00000000042 * cos( 5.13054174413 +    27433.88921587499*t);
+   emb_y_0+=     0.00000000042 * cos( 1.32093818083 +    21548.96236929180*t);
+   emb_y_0+=     0.00000000045 * cos( 3.22031193914 +    71519.29714327680*t);
+   emb_y_0+=     0.00000000043 * cos( 2.75467576214 +    27278.46881644080*t);
+   emb_y_0+=     0.00000000042 * cos( 2.68353478578 +    28230.18722269139*t);
+   emb_y_0+=     0.00000000051 * cos( 2.95623880108 +     7761.94242405580*t);
+   emb_y_0+=     0.00000000041 * cos( 5.14394847571 +     3735.23831175900*t);
+   emb_y_0+=     0.00000000042 * cos( 3.66392136716 +    26084.02180621620*t);
+   emb_y_0+=     0.00000000047 * cos( 6.03338630990 +     4708.23004870920*t);
+   emb_y_0+=     0.00000000056 * cos( 5.67442109710 +    88041.95685927900*t);
+   emb_y_0+=     0.00000000056 * cos( 5.36734540218 +    28663.83165026540*t);
+   emb_y_0+=     0.00000000040 * cos( 1.07842393153 +    28628.33622609960*t);
+   emb_y_0+=     0.00000000039 * cos( 0.71682305362 +    34570.06633485260*t);
+   emb_y_0+=     0.00000000056 * cos( 1.60418809786 +    82534.40362061159*t);
+   emb_y_0+=     0.00000000052 * cos( 1.15000260118 +    40398.18991926600*t);
+   emb_y_0+=     0.00000000048 * cos( 0.93182108526 +    12662.13092720060*t);
+   emb_y_0+=     0.00000000041 * cos( 4.95215656542 +    37057.57749256620*t);
+   emb_y_0+=     0.00000000052 * cos( 2.84085536047 +    12571.67462429020*t);
+   emb_y_0+=     0.00000000043 * cos( 3.12714109961 +     7322.10246078180*t);
+   emb_y_0+=     0.00000000041 * cos( 1.98689255259 +      664.75604513000*t);
+   emb_y_0+=     0.00000000043 * cos( 4.45441019609 +    19004.64794940840*t);
+   emb_y_0+=     0.00000000044 * cos( 3.63007864944 +     6140.89722295520*t);
+   emb_y_0+=     0.00000000044 * cos( 6.16004757680 +     6425.25447702760*t);
+   emb_y_0+=     0.00000000044 * cos( 3.57757676616 +    49515.38250840700*t);
+   emb_y_0+=     0.00000000039 * cos( 5.09443272912 +    22594.05489571199*t);
+   emb_y_0+=     0.00000000055 * cos( 5.97876206707 +    80181.53746683979*t);
+   emb_y_0+=     0.00000000054 * cos( 0.38609427782 +     6560.11084373280*t);
+   emb_y_0+=     0.00000000054 * cos( 3.12084664124 +     6006.04085625000*t);
+   emb_y_0+=     0.00000000045 * cos( 0.83262497188 +   142861.39123625819*t);
+   emb_y_0+=     0.00000000047 * cos( 5.83395711112 +    24499.51966065240*t);
+   emb_y_0+=     0.00000000043 * cos( 6.07854396898 +     7026.06591052400*t);
+   emb_y_0+=     0.00000000047 * cos( 2.57216284949 +    11919.14086666800*t);
+   emb_y_0+=     0.00000000047 * cos( 2.93748065634 +    24485.29256665079*t);
+   emb_y_0+=     0.00000000050 * cos( 3.13873979564 +    92747.68916682259*t);
+   emb_y_0+=     0.00000000043 * cos( 4.24359473122 +     6359.34192126700*t);
+   emb_y_0+=     0.00000000049 * cos( 2.11824959319 +    66813.56483573320*t);
+   emb_y_0+=     0.00000000042 * cos( 3.13003412642 +     6231.79551212899*t);
+   emb_y_0+=     0.00000000042 * cos( 0.37690679264 +     6334.35618785381*t);
+   emb_y_0+=     0.00000000052 * cos( 1.96232281613 +    11904.91877320180*t);
+   emb_y_0+=     0.00000000040 * cos( 6.22187276386 +     7075.85073845880*t);
+   emb_y_0+=     0.00000000040 * cos( 4.78395267597 +     7122.04513774180*t);
+   emb_y_0+=     0.00000000040 * cos( 5.00617355027 +     5444.10656224100*t);
+   emb_y_0+=     0.00000000038 * cos( 0.67444569484 +    44933.24935619040*t);
+   emb_y_0+=     0.00000000037 * cos( 4.80355711691 +     6850.79448772180*t);
+   emb_y_0+=     0.00000000037 * cos( 4.98656910933 +     5715.35721226100*t);
+   emb_y_0+=     0.00000000050 * cos( 0.35643037555 +    21947.11137270000*t);
+   emb_y_0+=     0.00000000039 * cos( 2.97136332005 +    21150.81336588360*t);
+   emb_y_0+=     0.00000000042 * cos( 2.21951419227 +     6328.21706962800*t);
+   emb_y_0+=     0.00000000042 * cos( 1.28742672679 +     6237.93463035480*t);
+   emb_y_0+=     0.00000000040 * cos( 5.21447072458 +     3215.13638064320*t);
+   emb_y_0+=     0.00000000039 * cos( 3.85122558089 +    30774.50164257480*t);
+   emb_y_0+=     0.00000000046 * cos( 3.45581538922 +    12146.66705610760*t);
+
+   double emb_y_1=0.0;
+   emb_y_1+=     0.00093046317 * cos( 0.00000000000 +        0.00000000000*t);
+   emb_y_1+=     0.00051506609 * cos( 4.43180499286 +    12566.15169998280*t);
+   emb_y_1+=     0.00001290800 * cos( 4.38860548540 +    18849.22754997420*t);
+   emb_y_1+=     0.00000464550 * cos( 5.82729912952 +     6283.07584999140*t);
+   emb_y_1+=     0.00000212689 * cos( 0.16300556918 +     6279.55273164240*t);
+   emb_y_1+=     0.00000212533 * cos( 3.34400595407 +     6286.59896834040*t);
+   emb_y_1+=     0.00000062345 * cos( 5.07377354827 +     4705.73230754360*t);
+   emb_y_1+=     0.00000059794 * cos( 2.24100907272 +     6256.77753019160*t);
+   emb_y_1+=     0.00000059441 * cos( 1.26619990626 +     6309.37416979120*t);
+   emb_y_1+=     0.00000048987 * cos( 3.64166577835 +      775.52261132400*t);
+   emb_y_1+=     0.00000042746 * cos( 5.15177425824 +     1059.38193018920*t);
+   emb_y_1+=     0.00000046642 * cos( 4.71575215800 +     7860.41939243920*t);
+   emb_y_1+=     0.00000037274 * cos( 0.66041323804 +     5753.38488489680*t);
+   emb_y_1+=     0.00000035625 * cos( 6.15462846966 +     5884.92684658320*t);
+   emb_y_1+=     0.00000035373 * cos( 2.90052012152 +     6812.76681508600*t);
+   emb_y_1+=     0.00000032157 * cos( 3.62669702015 +     6681.22485339960*t);
+   emb_y_1+=     0.00000028763 * cos( 4.34539355621 +    25132.30339996560*t);
+   emb_y_1+=     0.00000028447 * cos( 5.86215103214 +     6127.65545055720*t);
+   emb_y_1+=     0.00000027502 * cos( 3.92641632542 +     6438.49624942560*t);
+   emb_y_1+=     0.00000024830 * cos( 1.35190115953 +     5486.77784317500*t);
+   emb_y_1+=     0.00000020618 * cos( 2.14772396647 +     7079.37385680780*t);
+   emb_y_1+=     0.00000019531 * cos( 1.29958310594 +     5507.55323866740*t);
+   emb_y_1+=     0.00000018447 * cos( 6.17330532651 +    11790.62908865880*t);
+   emb_y_1+=     0.00000016520 * cos( 4.63328039489 +    11506.76976979360*t);
+   emb_y_1+=     0.00000016754 * cos( 2.24862220282 +     7058.59846131540*t);
+   emb_y_1+=     0.00000014560 * cos( 4.40425761361 +     6290.18939699220*t);
+   emb_y_1+=     0.00000014426 * cos( 5.39226330857 +     6275.96230299060*t);
+   emb_y_1+=     0.00000014186 * cos( 6.18547219419 +      796.29800681640*t);
+   emb_y_1+=     0.00000012666 * cos( 5.73231393752 +     4694.00295470760*t);
+   emb_y_1+=     0.00000012966 * cos( 5.10338935824 +        7.11354700080*t);
+   emb_y_1+=     0.00000010100 * cos( 1.72464513593 +     3738.76143010800*t);
+   emb_y_1+=     0.00000009427 * cos( 3.08115743044 +     6069.77675455340*t);
+   emb_y_1+=     0.00000010425 * cos( 2.69531177468 +     6282.09552892320*t);
+   emb_y_1+=     0.00000010425 * cos( 0.81162914438 +     6284.05617105960*t);
+   emb_y_1+=     0.00000009576 * cos( 0.32459396888 +     4136.91043351620*t);
+   emb_y_1+=     0.00000009043 * cos( 0.43888880403 +     6496.37494542940*t);
+   emb_y_1+=     0.00000008545 * cos( 4.77137497641 +     1194.44701022460*t);
+   emb_y_1+=     0.00000008199 * cos( 5.96838316358 +     1589.07289528380*t);
+   emb_y_1+=     0.00000006291 * cos( 3.18312019173 +     8429.24126646660*t);
+   emb_y_1+=     0.00000006133 * cos( 1.78033290355 +     8827.39026987480*t);
+   emb_y_1+=     0.00000005524 * cos( 3.79842135217 +     4933.20844033260*t);
+   emb_y_1+=     0.00000007316 * cos( 5.20404355405 +     4535.05943692440*t);
+   emb_y_1+=     0.00000005146 * cos( 3.57146870511 +    11769.85369316640*t);
+   emb_y_1+=     0.00000006461 * cos( 2.67840586109 +     5088.62883976680*t);
+   emb_y_1+=     0.00000005270 * cos( 0.94095218349 +     3154.68708489560*t);
+   emb_y_1+=     0.00000004829 * cos( 5.32876538500 +     6040.34724601740*t);
+   emb_y_1+=     0.00000004985 * cos( 3.05697984231 +    12569.67481833180*t);
+   emb_y_1+=     0.00000004709 * cos( 2.39292099451 +     5331.35744374080*t);
+   emb_y_1+=     0.00000004601 * cos( 4.45296532710 +     6525.80445396540*t);
+   emb_y_1+=     0.00000005607 * cos( 0.98720797401 +     5729.50644714900*t);
+   emb_y_1+=     0.00000004081 * cos( 5.94191743597 +     7632.94325965020*t);
+   emb_y_1+=     0.00000004394 * cos( 4.75948184840 +    12559.03815298200*t);
+   emb_y_1+=     0.00000005053 * cos( 4.55162942081 +     8031.09226305840*t);
+   emb_y_1+=     0.00000005077 * cos( 0.73549620962 +     7477.52286021600*t);
+   emb_y_1+=     0.00000003779 * cos( 1.07385171253 +     7234.79425624200*t);
+   emb_y_1+=     0.00000004974 * cos( 2.50029147401 +     6836.64525283380*t);
+   emb_y_1+=     0.00000003614 * cos( 2.52455299148 +     3930.20969621960*t);
+   emb_y_1+=     0.00000003482 * cos( 3.82652198801 +    10977.07880469900*t);
+   emb_y_1+=     0.00000003576 * cos( 5.30219799481 +     4164.31198961300*t);
+   emb_y_1+=     0.00000003376 * cos( 5.00213145251 +    11371.70468975820*t);
+   emb_y_1+=     0.00000004290 * cos( 3.36736739298 +     1592.59601363280*t);
+   emb_y_1+=     0.00000003406 * cos( 0.83525752781 +     3128.38876509580*t);
+   emb_y_1+=     0.00000003124 * cos( 3.17143982984 +     7084.89678111520*t);
+   emb_y_1+=     0.00000002752 * cos( 3.67389621489 +     7342.45778018060*t);
+   emb_y_1+=     0.00000002778 * cos( 0.26817898615 +     1748.01641306700*t);
+   emb_y_1+=     0.00000002694 * cos( 4.70989138148 +      529.69096509460*t);
+   emb_y_1+=     0.00000002820 * cos( 0.23146914320 +     5223.69391980220*t);
+   emb_y_1+=     0.00000003334 * cos( 3.48823993118 +      801.82093112380*t);
+   emb_y_1+=     0.00000003328 * cos( 5.83812544995 +      155.42039943420*t);
+   emb_y_1+=     0.00000003008 * cos( 0.67203387304 +     8635.94200376320*t);
+   emb_y_1+=     0.00000002348 * cos( 5.19665664165 +     2146.16541647520*t);
+   emb_y_1+=     0.00000002307 * cos( 2.98930142346 +    17260.15465469040*t);
+   emb_y_1+=     0.00000002778 * cos( 6.17392405934 +     5481.25491886760*t);
+   emb_y_1+=     0.00000002046 * cos( 1.96605971472 +     1990.74501704100*t);
+   emb_y_1+=     0.00000002670 * cos( 1.00151747879 +       26.29831979980*t);
+   emb_y_1+=     0.00000002000 * cos( 3.03124266219 +      951.71840625060*t);
+   emb_y_1+=     0.00000002162 * cos( 0.33268106563 +      536.80451209540*t);
+   emb_y_1+=     0.00000001830 * cos( 4.43806423544 +      553.56940284240*t);
+   emb_y_1+=     0.00000001762 * cos( 3.84931185617 +     4690.47983635860*t);
+   emb_y_1+=     0.00000002033 * cos( 0.80949219978 +      522.57741809380*t);
+   emb_y_1+=     0.00000002430 * cos( 1.64090263419 +     1349.86740965880*t);
+   emb_y_1+=     0.00000001902 * cos( 2.66567222573 +     9437.76293488700*t);
+   emb_y_1+=     0.00000001940 * cos( 3.91085273629 +     5216.58037280140*t);
+   emb_y_1+=     0.00000001708 * cos( 5.00417878382 +    10988.80815753500*t);
+   emb_y_1+=     0.00000001762 * cos( 2.39542835845 +      398.14900340820*t);
+   emb_y_1+=     0.00000001790 * cos( 3.48655562486 +     5230.80746680300*t);
+   emb_y_1+=     0.00000001524 * cos( 5.81370227284 +      426.59819087600*t);
+   emb_y_1+=     0.00000001555 * cos( 2.09924417685 +    13095.84266507740*t);
+   emb_y_1+=     0.00000001779 * cos( 0.23737949225 +    12562.62858163380*t);
+   emb_y_1+=     0.00000001450 * cos( 5.71201368260 +     6262.72053059260*t);
+   emb_y_1+=     0.00000001450 * cos( 4.07811254364 +     6303.43116939020*t);
+   emb_y_1+=     0.00000001768 * cos( 0.08826500577 +    10973.55568635000*t);
+   emb_y_1+=     0.00000001373 * cos( 5.60523185373 +     7875.67186362420*t);
+   emb_y_1+=     0.00000001428 * cos( 3.29442777483 +    12721.57209941700*t);
+   emb_y_1+=     0.00000001247 * cos( 0.29212365805 +     2118.76386037840*t);
+   emb_y_1+=     0.00000001230 * cos( 2.53131850803 +    18319.53658487960*t);
+   emb_y_1+=     0.00000001141 * cos( 1.92147452671 +     1066.49547719000*t);
+   emb_y_1+=     0.00000001140 * cos( 1.21410806683 +     2942.46342329160*t);
+   emb_y_1+=     0.00000001045 * cos( 2.17535721280 +      242.72860397400*t);
+   emb_y_1+=     0.00000000973 * cos( 4.95437649984 +      639.89728631400*t);
+   emb_y_1+=     0.00000000969 * cos( 1.29405045734 +     5643.17856367740*t);
+   emb_y_1+=     0.00000001064 * cos( 4.40566643488 +    15720.83878487840*t);
+   emb_y_1+=     0.00000000962 * cos( 0.55758587779 +     2388.89402044920*t);
+   emb_y_1+=     0.00000000943 * cos( 4.22062437022 +    20426.57109242200*t);
+   emb_y_1+=     0.00000000869 * cos( 5.80860615243 +    16496.36139620240*t);
+   emb_y_1+=     0.00000000874 * cos( 1.43036958888 +    10575.40668294180*t);
+   emb_y_1+=     0.00000001088 * cos( 5.77656856686 +     5760.49843189760*t);
+   emb_y_1+=     0.00000000889 * cos( 4.77820137734 +     6288.59877429880*t);
+   emb_y_1+=     0.00000000889 * cos( 5.01192484889 +     6277.55292568400*t);
+   emb_y_1+=     0.00000001018 * cos( 2.66299766067 +    12539.85338018300*t);
+   emb_y_1+=     0.00000001017 * cos( 4.02093017183 +     6805.65326808520*t);
+   emb_y_1+=     0.00000000794 * cos( 3.79329079940 +    14314.16811304980*t);
+   emb_y_1+=     0.00000000795 * cos( 3.39951828424 +     6279.78949257360*t);
+   emb_y_1+=     0.00000000795 * cos( 0.10742263482 +     6286.36220740920*t);
+   emb_y_1+=     0.00000000759 * cos( 4.98301749094 +     1577.34354244780*t);
+   emb_y_1+=     0.00000000757 * cos( 3.81754073272 +     5856.47765911540*t);
+   emb_y_1+=     0.00000000800 * cos( 3.18366218591 +    12029.34718788740*t);
+   emb_y_1+=     0.00000000716 * cos( 2.28687269516 +     9623.68827669120*t);
+   emb_y_1+=     0.00000000698 * cos( 4.94363229969 +     6148.01076995600*t);
+   emb_y_1+=     0.00000000907 * cos( 0.70119368826 +    13367.97263110660*t);
+   emb_y_1+=     0.00000000672 * cos( 3.99153941403 +    12964.30070339100*t);
+   emb_y_1+=     0.00000000678 * cos( 4.84647216304 +     6418.14093002680*t);
+   emb_y_1+=     0.00000000671 * cos( 6.00179402654 +     6709.67404086740*t);
+   emb_y_1+=     0.00000000657 * cos( 1.80163596347 +    13119.72110282519*t);
+   emb_y_1+=     0.00000000662 * cos( 4.04031103640 +    11015.10647733480*t);
+   emb_y_1+=     0.00000000744 * cos( 2.68741160780 +    12043.57428188900*t);
+   emb_y_1+=     0.00000000673 * cos( 2.45322928062 +    14712.31711645800*t);
+   emb_y_1+=     0.00000000696 * cos( 2.53877129238 +    16730.46368959580*t);
+   emb_y_1+=     0.00000000671 * cos( 0.40543909773 +    13517.87010623340*t);
+   emb_y_1+=     0.00000000604 * cos( 5.39462125376 +    14143.49524243060*t);
+   emb_y_1+=     0.00000000754 * cos( 5.24655445271 +    13916.01910964160*t);
+   emb_y_1+=     0.00000000629 * cos( 3.42493725594 +     4686.88940770680*t);
+   emb_y_1+=     0.00000000610 * cos( 5.19955844918 +     6819.88036208680*t);
+   emb_y_1+=     0.00000000667 * cos( 4.87151372910 +     3634.62102451840*t);
+   emb_y_1+=     0.00000000609 * cos( 4.57544043999 +     5746.27133789600*t);
+   emb_y_1+=     0.00000000620 * cos( 4.30310535016 +    31415.37924995700*t);
+   emb_y_1+=     0.00000000553 * cos( 0.78994516254 +     6306.95428773920*t);
+   emb_y_1+=     0.00000000553 * cos( 2.71699575652 +     6259.19741224360*t);
+   emb_y_1+=     0.00000000524 * cos( 3.96543945899 +    12139.55350910680*t);
+   emb_y_1+=     0.00000000514 * cos( 3.87793148429 +     4701.11650170840*t);
+   emb_y_1+=     0.00000000553 * cos( 4.81568395217 +    12036.46073488820*t);
+   emb_y_1+=     0.00000000495 * cos( 2.00987395637 +    13362.44970679920*t);
+   emb_y_1+=     0.00000000562 * cos( 4.76145666018 +     4292.33083295040*t);
+   emb_y_1+=     0.00000000481 * cos( 2.57118063835 +     6286.95718534940*t);
+   emb_y_1+=     0.00000000481 * cos( 0.93576028071 +     6279.19451463340*t);
+   emb_y_1+=     0.00000000445 * cos( 1.68959289911 +    11499.65622279280*t);
+   emb_y_1+=     0.00000000596 * cos( 1.08038043695 +    10447.38783960440*t);
+   emb_y_1+=     0.00000000452 * cos( 3.98330238308 +    12573.26524698360*t);
+   emb_y_1+=     0.00000000499 * cos( 5.79985861523 +    12410.73130054860*t);
+   emb_y_1+=     0.00000000440 * cos( 0.58121786395 +    24072.92146977640*t);
+   emb_y_1+=     0.00000000575 * cos( 1.39410252291 +    12592.45001978260*t);
+   emb_y_1+=     0.00000000404 * cos( 0.56451183179 +     5863.59120611620*t);
+   emb_y_1+=     0.00000000398 * cos( 4.20035189119 +     8273.82086703240*t);
+   emb_y_1+=     0.00000000418 * cos( 2.76069368004 +    10177.25767953360*t);
+   emb_y_1+=     0.00000000379 * cos( 3.05384870302 +     6386.16862421000*t);
+   emb_y_1+=     0.00000000379 * cos( 0.45309221604 +     6179.98307577280*t);
+   emb_y_1+=     0.00000000377 * cos( 4.33269290931 +     7872.14874527520*t);
+   emb_y_1+=     0.00000000419 * cos( 4.05140116867 +     5327.47610838280*t);
+   emb_y_1+=     0.00000000383 * cos( 5.11659434377 +     6319.10371666880*t);
+   emb_y_1+=     0.00000000383 * cos( 4.67353188247 +     6247.04798331399*t);
+   emb_y_1+=     0.00000000457 * cos( 2.62120808067 +     6076.89030155420*t);
+   emb_y_1+=     0.00000000457 * cos( 0.88573283839 +     6489.26139842860*t);
+   emb_y_1+=     0.00000000378 * cos( 5.80861101902 +     7238.67559160000*t);
+   emb_y_1+=     0.00000000344 * cos( 4.99907349425 +    18451.07854656599*t);
+   emb_y_1+=     0.00000000334 * cos( 0.56215917990 +    18073.70493865020*t);
+   emb_y_1+=     0.00000000453 * cos( 5.43513758472 +     2787.04302385740*t);
+   emb_y_1+=     0.00000000364 * cos( 0.67991079022 +     6321.10352262720*t);
+   emb_y_1+=     0.00000000364 * cos( 2.82703012884 +     6245.04817735560*t);
+   emb_y_1+=     0.00000000342 * cos( 4.52940007807 +     6303.85124548380*t);
+   emb_y_1+=     0.00000000342 * cos( 5.26072614817 +     6262.30045449900*t);
+   emb_y_1+=     0.00000000319 * cos( 4.83173153249 +    11926.25441366880*t);
+   emb_y_1+=     0.00000000318 * cos( 6.03589239585 +     6311.52503745920*t);
+   emb_y_1+=     0.00000000318 * cos( 3.75423383039 +     6254.62666252360*t);
+   emb_y_1+=     0.00000000315 * cos( 1.09950918957 +    12323.42309600880*t);
+   emb_y_1+=     0.00000000310 * cos( 4.02477182017 +    10213.28554621100*t);
+   emb_y_1+=     0.00000000311 * cos( 5.94074472453 +     2352.86615377180*t);
+   emb_y_1+=     0.00000000304 * cos( 0.43890915612 +    12779.45079542080*t);
+   emb_y_1+=     0.00000000343 * cos( 1.36470897271 +     6267.82337880639*t);
+   emb_y_1+=     0.00000000343 * cos( 2.14223194635 +     6298.32832117640*t);
+   emb_y_1+=     0.00000000398 * cos( 4.83558508405 +     1551.04522264800*t);
+   emb_y_1+=     0.00000000358 * cos( 2.52713972155 +     1052.26838318840*t);
+   emb_y_1+=     0.00000000297 * cos( 2.07974536467 +     5429.87946823940*t);
+   emb_y_1+=     0.00000000332 * cos( 2.13791551771 +      213.29909543800*t);
+   emb_y_1+=     0.00000000362 * cos( 1.16895887442 +    10021.83728009940*t);
+   emb_y_1+=     0.00000000334 * cos( 1.16416729842 +    17782.73207278420*t);
+   emb_y_1+=     0.00000000295 * cos( 2.48770580169 +     6702.56049386660*t);
+   emb_y_1+=     0.00000000321 * cos( 0.66776723307 +    17796.95916678580*t);
+   emb_y_1+=     0.00000000369 * cos( 4.78231597870 +     4732.03062734340*t);
+   emb_y_1+=     0.00000000288 * cos( 3.73055532740 +      220.41264243880*t);
+   emb_y_1+=     0.00000000263 * cos( 2.60909847321 +      955.59974160860*t);
+   emb_y_1+=     0.00000000305 * cos( 0.57265325631 +     5642.19824260920*t);
+   emb_y_1+=     0.00000000258 * cos( 4.38103579313 +      206.18554843720*t);
+   emb_y_1+=     0.00000000287 * cos( 0.31572653335 +     9917.69687450980*t);
+   emb_y_1+=     0.00000000261 * cos( 1.93782216260 +    16200.77272450120*t);
+   emb_y_1+=     0.00000000234 * cos( 2.50307529331 +     1596.18644228460*t);
+   emb_y_1+=     0.00000000282 * cos( 2.93181184874 +    13341.67431130680*t);
+   emb_y_1+=     0.00000000262 * cos( 5.15409122759 +    14945.31617355440*t);
+   emb_y_1+=     0.00000000276 * cos( 0.59107321564 +    12168.00269657460*t);
+   emb_y_1+=     0.00000000217 * cos( 2.74597813026 +     2544.31441988340*t);
+   emb_y_1+=     0.00000000259 * cos( 5.53112324599 +     3894.18182954220*t);
+   emb_y_1+=     0.00000000216 * cos( 4.02991604351 +     3185.19202726560*t);
+   emb_y_1+=     0.00000000250 * cos( 3.04516511604 +     6923.95345737360*t);
+   emb_y_1+=     0.00000000193 * cos( 0.78779752820 +      640.87760738220*t);
+   emb_y_1+=     0.00000000203 * cos( 2.82417195790 +     6062.66320755260*t);
+   emb_y_1+=     0.00000000203 * cos( 0.68276896116 +     6503.48849243020*t);
+   emb_y_1+=     0.00000000188 * cos( 4.83298725434 +    18209.33026366019*t);
+   emb_y_1+=     0.00000000184 * cos( 2.63875516800 +    12566.21901028560*t);
+   emb_y_1+=     0.00000000213 * cos( 4.09959554852 +     9779.10867612540*t);
+   emb_y_1+=     0.00000000180 * cos( 4.19738483213 +      853.19638175200*t);
+   emb_y_1+=     0.00000000248 * cos( 5.62050900898 +    17789.84561978500*t);
+   emb_y_1+=     0.00000000187 * cos( 3.02088056063 +     6276.02961329339*t);
+   emb_y_1+=     0.00000000187 * cos( 0.48606035843 +     6290.12208668940*t);
+   emb_y_1+=     0.00000000168 * cos( 6.02158735674 +      110.20632121940*t);
+   emb_y_1+=     0.00000000184 * cos( 3.08285813784 +    12566.08438968000*t);
+   emb_y_1+=     0.00000000230 * cos( 0.66694562430 +    12567.13202105100*t);
+   emb_y_1+=     0.00000000166 * cos( 5.90725172950 +      433.71173787680*t);
+   emb_y_1+=     0.00000000184 * cos( 1.80970202711 +     6037.24420376200*t);
+   emb_y_1+=     0.00000000159 * cos( 0.31775301914 +     3496.03282613400*t);
+   emb_y_1+=     0.00000000179 * cos( 1.27622455391 +        3.52311834900*t);
+   emb_y_1+=     0.00000000156 * cos( 0.77329171048 +     2648.45482547300*t);
+   emb_y_1+=     0.00000000182 * cos( 0.19927362554 +     5849.36411211460*t);
+   emb_y_1+=     0.00000000175 * cos( 1.69691956469 +     6528.90749622080*t);
+   emb_y_1+=     0.00000000154 * cos( 5.68366565747 +     6112.40297937220*t);
+   emb_y_1+=     0.00000000154 * cos( 4.10646056877 +     6453.74872061060*t);
+   emb_y_1+=     0.00000000162 * cos( 0.40509792796 +    18875.52586977400*t);
+   emb_y_1+=     0.00000000163 * cos( 2.96167567065 +    12352.85260454480*t);
+   emb_y_1+=     0.00000000146 * cos( 2.48714617179 +      149.56319713460*t);
+   emb_y_1+=     0.00000000143 * cos( 0.06272029982 +    19651.04848109800*t);
+   emb_y_1+=     0.00000000144 * cos( 2.19098385924 +    18635.92845453620*t);
+   emb_y_1+=     0.00000000141 * cos( 5.63574749070 +    10969.96525769820*t);
+   emb_y_1+=     0.00000000146 * cos( 0.42084370397 +    13760.59871020740*t);
+   emb_y_1+=     0.00000000161 * cos( 6.06631099179 +     8662.24032356300*t);
+   emb_y_1+=     0.00000000142 * cos( 2.91677439721 +     4157.19844261220*t);
+   emb_y_1+=     0.00000000150 * cos( 4.37920284921 +     3340.61242669980*t);
+   emb_y_1+=     0.00000000125 * cos( 0.78764990758 +       20.77539549240*t);
+   emb_y_1+=     0.00000000123 * cos( 4.97498024339 +     6172.86952877200*t);
+   emb_y_1+=     0.00000000123 * cos( 4.81514598285 +     6393.28217121080*t);
+   emb_y_1+=     0.00000000135 * cos( 4.59213390099 +    10984.19235169980*t);
+   emb_y_1+=     0.00000000117 * cos( 2.90632080371 +    18852.75066832320*t);
+   emb_y_1+=     0.00000000118 * cos( 5.67650820416 +    17256.63153634140*t);
+   emb_y_1+=     0.00000000147 * cos( 4.60612631701 +    17654.78053974960*t);
+   emb_y_1+=     0.00000000123 * cos( 2.85759760823 +    12565.17137891460*t);
+   emb_y_1+=     0.00000000107 * cos( 5.46942612515 +     9380.95967271720*t);
+   emb_y_1+=     0.00000000119 * cos( 0.71211331875 +    10419.98628350760*t);
+   emb_y_1+=     0.00000000105 * cos( 5.13254223053 +      632.78373931320*t);
+   emb_y_1+=     0.00000000103 * cos( 4.95941330270 +     6271.34649715540*t);
+   emb_y_1+=     0.00000000103 * cos( 4.83071292354 +     6294.80520282740*t);
+   emb_y_1+=     0.00000000102 * cos( 5.99360230997 +     4804.20927592700*t);
+   emb_y_1+=     0.00000000114 * cos( 3.78885344182 +     4171.42553661380*t);
+   emb_y_1+=     0.00000000113 * cos( 4.29988480427 +    23013.53953958720*t);
+   emb_y_1+=     0.00000000112 * cos( 0.61034141917 +    17253.04110768959*t);
+   emb_y_1+=     0.00000000111 * cos( 2.09128320599 +     1581.95934828300*t);
+   emb_y_1+=     0.00000000113 * cos( 0.97983808350 +    16858.48253293320*t);
+   emb_y_1+=     0.00000000102 * cos( 4.09962394648 +     6549.68289171320*t);
+   emb_y_1+=     0.00000000102 * cos( 5.69050227976 +     6016.46880826960*t);
+   emb_y_1+=     0.00000000106 * cos( 0.87687085688 +    23543.23050468179*t);
+   emb_y_1+=     0.00000000121 * cos( 4.76575782873 +    13625.53363017200*t);
+   emb_y_1+=     0.00000000090 * cos( 2.11774047359 +       76.26607127560*t);
+   emb_y_1+=     0.00000000088 * cos( 1.18967345782 +    17298.18232732620*t);
+   emb_y_1+=     0.00000000087 * cos( 1.12941525302 +    17267.26820169119*t);
+   emb_y_1+=     0.00000000092 * cos( 1.84158032509 +    22805.73556599360*t);
+   emb_y_1+=     0.00000000086 * cos( 4.36870119289 +    12808.88030395680*t);
+   emb_y_1+=     0.00000000111 * cos( 0.85102422912 +    11513.88331679440*t);
+   emb_y_1+=     0.00000000095 * cos( 6.05597457730 +    12345.73905754400*t);
+   emb_y_1+=     0.00000000106 * cos( 3.57199937061 +    11216.28429032400*t);
+   emb_y_1+=     0.00000000087 * cos( 1.36534411656 +     8982.81066930900*t);
+   emb_y_1+=     0.00000000095 * cos( 0.92966303480 +    12012.58229714040*t);
+   emb_y_1+=     0.00000000081 * cos( 2.27691876526 +    11614.43329373220*t);
+   emb_y_1+=     0.00000000084 * cos( 5.85994160066 +     6255.67429389460*t);
+   emb_y_1+=     0.00000000084 * cos( 3.93018462558 +     6310.47740608820*t);
+   emb_y_1+=     0.00000000080 * cos( 5.05950625443 +     6922.97313630540*t);
+   emb_y_1+=     0.00000000085 * cos( 1.08198764342 +     5636.06501667660*t);
+   emb_y_1+=     0.00000000084 * cos( 4.64040201203 +    10440.27429260360*t);
+   emb_y_1+=     0.00000000089 * cos( 6.25985265406 +    16723.35014259500*t);
+   emb_y_1+=     0.00000000080 * cos( 2.62622474570 +     3583.34103067380*t);
+   emb_y_1+=     0.00000000091 * cos( 5.59316682916 +    10818.13528691580*t);
+   emb_y_1+=     0.00000000091 * cos( 1.43612727660 +    15671.08175940660*t);
+   emb_y_1+=     0.00000000076 * cos( 2.85409687552 +    11925.27409260060*t);
+   emb_y_1+=     0.00000000096 * cos( 1.47205677028 +     3097.88382272579*t);
+   emb_y_1+=     0.00000000104 * cos( 0.70746530546 +     5650.29211067820*t);
+   emb_y_1+=     0.00000000086 * cos( 5.52511956951 +    12359.96615154560*t);
+   emb_y_1+=     0.00000000080 * cos( 6.03205499783 +    13521.75144159140*t);
+   emb_y_1+=     0.00000000090 * cos( 4.45214765055 +     6195.76764545159*t);
+   emb_y_1+=     0.00000000090 * cos( 5.33797857569 +     6370.38405453120*t);
+   emb_y_1+=     0.00000000082 * cos( 4.91006846138 +    18842.11400297339*t);
+   emb_y_1+=     0.00000000073 * cos( 4.52976823745 +      949.17560896980*t);
+   emb_y_1+=     0.00000000083 * cos( 3.74911363362 +    18415.51581209739*t);
+   emb_y_1+=     0.00000000102 * cos( 4.43384315891 +     3104.93005942380*t);
+   emb_y_1+=     0.00000000074 * cos( 5.90055438198 +     5120.60114558360*t);
+   emb_y_1+=     0.00000000093 * cos( 2.53140640111 +    16460.33352952499*t);
+   emb_y_1+=     0.00000000084 * cos( 5.92514029201 +     2379.16447357160*t);
+   emb_y_1+=     0.00000000075 * cos( 0.05752343047 +    30355.99731976780*t);
+   emb_y_1+=     0.00000000076 * cos( 3.26399401386 +    18429.74290609900*t);
+   emb_y_1+=     0.00000000067 * cos( 4.00530001957 +    16062.18452611680*t);
+   emb_y_1+=     0.00000000066 * cos( 4.81634064130 +    10454.50138660520*t);
+   emb_y_1+=     0.00000000087 * cos( 3.52675041577 +     3903.91137641980*t);
+   emb_y_1+=     0.00000000078 * cos( 2.74560762552 +     8671.96987044060*t);
+   emb_y_1+=     0.00000000063 * cos( 5.70772167485 +    24356.78078864160*t);
+   emb_y_1+=     0.00000000073 * cos( 6.08379700055 +     9388.00590941520*t);
+   emb_y_1+=     0.00000000060 * cos( 1.94470153573 +    34570.06633485260*t);
+   emb_y_1+=     0.00000000056 * cos( 3.76314176838 +    12455.94537876340*t);
+   emb_y_1+=     0.00000000055 * cos( 1.32127703933 +    23539.70738633280*t);
+   emb_y_1+=     0.00000000055 * cos( 5.60680941684 +    15664.03552270859*t);
+   emb_y_1+=     0.00000000056 * cos( 0.60887110989 +    11764.33076885900*t);
+   emb_y_1+=     0.00000000061 * cos( 4.86427999867 +    28237.23345938940*t);
+   emb_y_1+=     0.00000000065 * cos( 4.29965528550 +     7335.34423317980*t);
+   emb_y_1+=     0.00000000067 * cos( 4.20742811531 +    11933.36796066960*t);
+   emb_y_1+=     0.00000000072 * cos( 4.30216792995 +     6510.55198278040*t);
+   emb_y_1+=     0.00000000072 * cos( 5.48795829629 +     6055.59971720240*t);
+   emb_y_1+=     0.00000000053 * cos( 5.64466034451 +    17996.03116822220*t);
+   emb_y_1+=     0.00000000071 * cos( 4.94322519248 +    12545.37630449040*t);
+   emb_y_1+=     0.00000000061 * cos( 3.80625548858 +     6253.64634145540*t);
+   emb_y_1+=     0.00000000061 * cos( 5.98387073766 +     6312.50535852740*t);
+   emb_y_1+=     0.00000000052 * cos( 1.83796108842 +     5326.78669402080*t);
+   emb_y_1+=     0.00000000066 * cos( 0.79633834660 +    14919.01785375460*t);
+   emb_y_1+=     0.00000000053 * cos( 2.81005325604 +     2699.73481931760*t);
+   emb_y_1+=     0.00000000051 * cos( 2.33512659118 +    16522.65971600220*t);
+   emb_y_1+=     0.00000000049 * cos( 2.61656735675 +    35579.69123957000*t);
+   emb_y_1+=     0.00000000053 * cos( 3.71469727042 +     6357.85744855870*t);
+   emb_y_1+=     0.00000000053 * cos( 6.07542895582 +     6208.29425142410*t);
+   emb_y_1+=     0.00000000058 * cos( 6.24714677357 +    16737.57723659660*t);
+   emb_y_1+=     0.00000000049 * cos( 3.35591683551 +    21393.54196985760*t);
+   emb_y_1+=     0.00000000049 * cos( 5.31764383927 +     7018.95236352320*t);
+   emb_y_1+=     0.00000000049 * cos( 4.47248238697 +     5547.19933645960*t);
+   emb_y_1+=     0.00000000049 * cos( 1.02044068888 +    12416.58850284820*t);
+   emb_y_1+=     0.00000000050 * cos( 4.78756360288 +    26084.02180621620*t);
+   emb_y_1+=     0.00000000052 * cos( 0.61376591765 +    45585.17281218740*t);
+   emb_y_1+=     0.00000000051 * cos( 1.83598363802 +    23006.42599258639*t);
+   emb_y_1+=     0.00000000058 * cos( 4.03624047121 +    15141.39079431200*t);
+   emb_y_1+=     0.00000000047 * cos( 3.55960074091 +    22743.40937951640*t);
+   emb_y_1+=     0.00000000050 * cos( 5.52194731109 +    21424.46664430340*t);
+   emb_y_1+=     0.00000000047 * cos( 5.68901702449 +    24279.10701821359*t);
+   emb_y_1+=     0.00000000049 * cos( 0.09767152059 +     5749.86176654780*t);
+   emb_y_1+=     0.00000000049 * cos( 3.40926939847 +     6816.28993343500*t);
+   emb_y_1+=     0.00000000046 * cos( 1.88054905125 +    19378.91851506880*t);
+   emb_y_1+=     0.00000000047 * cos( 0.52217551707 +    24065.80792277559*t);
+   emb_y_1+=     0.00000000046 * cos( 4.16716226231 +   104351.61256629678*t);
+   emb_y_1+=     0.00000000050 * cos( 2.92366135772 +    24602.61243487099*t);
+   emb_y_1+=     0.00000000054 * cos( 4.75214506546 +     7834.12107263940*t);
+   emb_y_1+=     0.00000000049 * cos( 2.94403678329 +    29303.72893657940*t);
+   emb_y_1+=     0.00000000048 * cos( 1.95408707684 +     5244.04923920100*t);
+   emb_y_1+=     0.00000000043 * cos( 3.93857315871 +    12132.43996210600*t);
+   emb_y_1+=     0.00000000047 * cos( 5.81677586075 +    18208.34994259200*t);
+   emb_y_1+=     0.00000000049 * cos( 3.41896064981 +    29289.50184257780*t);
+   emb_y_1+=     0.00000000044 * cos( 2.06685104480 +     6175.41232605280*t);
+   emb_y_1+=     0.00000000044 * cos( 1.44008987426 +     6390.73937393000*t);
+   emb_y_1+=     0.00000000042 * cos( 6.18541909471 +    24705.70520908959*t);
+   emb_y_1+=     0.00000000041 * cos( 3.21613978391 +    17370.36097590980*t);
+   emb_y_1+=     0.00000000047 * cos( 1.17418881008 +    58953.14544329400*t);
+   emb_y_1+=     0.00000000041 * cos( 2.05518926769 +    23020.65308658799*t);
+   emb_y_1+=     0.00000000044 * cos( 3.20914273688 +    26735.94526221320*t);
+   emb_y_1+=     0.00000000045 * cos( 3.70228663504 +    54247.41313575040*t);
+   emb_y_1+=     0.00000000040 * cos( 4.57474127522 +    32243.49850088060*t);
+   emb_y_1+=     0.00000000040 * cos( 5.17182769558 +    24336.00539314920*t);
+   emb_y_1+=     0.00000000040 * cos( 3.62816477586 +    19804.82729158280*t);
+   emb_y_1+=     0.00000000039 * cos( 1.13462594128 +    18606.49894600020*t);
+   emb_y_1+=     0.00000000040 * cos( 4.71493835645 +     8401.83971036980*t);
+   emb_y_1+=     0.00000000041 * cos( 2.69655967341 +    16207.88627150200*t);
+   emb_y_1+=     0.00000000038 * cos( 5.56507429543 +    29088.81141598500*t);
+   emb_y_1+=     0.00000000038 * cos( 2.36481926067 +     6268.84875598980*t);
+   emb_y_1+=     0.00000000038 * cos( 1.14212165839 +     6297.30294399300*t);
+   emb_y_1+=     0.00000000041 * cos( 3.38234777816 +    15508.61512327440*t);
+   emb_y_1+=     0.00000000038 * cos( 4.82215691915 +     6171.88920770380*t);
+   emb_y_1+=     0.00000000038 * cos( 4.96796930709 +     6394.26249227900*t);
+   emb_y_1+=     0.00000000039 * cos( 1.75198804720 +    29826.30635467320*t);
+   emb_y_1+=     0.00000000037 * cos( 3.64776861935 +    20452.86941222180*t);
+   emb_y_1+=     0.00000000043 * cos( 6.01711305973 +    55022.93574707440*t);
+   emb_y_1+=     0.00000000037 * cos( 6.18577055273 +     4061.21921539440*t);
+   emb_y_1+=     0.00000000040 * cos( 2.02818434731 +    49515.38250840700*t);
+   emb_y_1+=     0.00000000049 * cos( 4.21086629146 +    18216.44381066100*t);
+   emb_y_1+=     0.00000000038 * cos( 4.59185538964 +    18202.21671665939*t);
+   emb_y_1+=     0.00000000035 * cos( 3.16929714068 +    26482.17080962440*t);
+   emb_y_1+=     0.00000000036 * cos( 1.56812014048 +     9924.81042151060*t);
+   emb_y_1+=     0.00000000039 * cos( 0.88019393819 +    24889.57479599160*t);
+   emb_y_1+=     0.00000000039 * cos( 4.63870568395 +    26709.64694241340*t);
+   emb_y_1+=     0.00000000036 * cos( 6.07069720518 +    25287.72379939980*t);
+   emb_y_1+=     0.00000000034 * cos( 6.10039920596 +     6147.45052498140*t);
+   emb_y_1+=     0.00000000034 * cos( 3.68972702028 +     6418.70117500140*t);
+   emb_y_1+=     0.00000000034 * cos( 2.54760886649 +    21228.39202354580*t);
+   emb_y_1+=     0.00000000034 * cos( 0.68922440987 +    33019.02111220460*t);
+   emb_y_1+=     0.00000000034 * cos( 0.40887359374 +    17271.88400752640*t);
+   emb_y_1+=     0.00000000034 * cos( 0.33424154418 +    29026.48522950779*t);
+   emb_y_1+=     0.00000000035 * cos( 1.63582495302 +     7322.10246078180*t);
+   emb_y_1+=     0.00000000044 * cos( 1.21937055747 +     3981.49003408200*t);
+   emb_y_1+=     0.00000000035 * cos( 3.00160924335 +    21954.15760939799*t);
+   emb_y_1+=     0.00000000034 * cos( 3.38052068808 +    38500.27603107220*t);
+   emb_y_1+=     0.00000000032 * cos( 2.81219711037 +    17686.75284556640*t);
+   emb_y_1+=     0.00000000033 * cos( 0.13837047156 +    34520.30930938080*t);
+   emb_y_1+=     0.00000000041 * cos( 2.52021125438 +    21202.09370374600*t);
+   emb_y_1+=     0.00000000037 * cos( 5.92836585991 +     7349.57132718140*t);
+   emb_y_1+=     0.00000000031 * cos( 2.79498926181 +    28230.18722269139*t);
+   emb_y_1+=     0.00000000038 * cos( 5.01277542891 +     9225.53927328300*t);
+   emb_y_1+=     0.00000000035 * cos( 2.49445060142 +     8584.66166590080*t);
+   emb_y_1+=     0.00000000038 * cos( 2.67041435749 +    15110.46611986620*t);
+   emb_y_1+=     0.00000000037 * cos( 1.84827551648 +    29296.61538957860*t);
+   emb_y_1+=     0.00000000035 * cos( 4.26342089717 +    11919.14086666800*t);
+   emb_y_1+=     0.00000000035 * cos( 2.59482262662 +    62883.35513951360*t);
+   emb_y_1+=     0.00000000035 * cos( 5.54313072444 +    29062.51309618520*t);
+   emb_y_1+=     0.00000000035 * cos( 5.96888200431 +    25685.87280280800*t);
+   emb_y_1+=     0.00000000030 * cos( 1.38932865248 +    12489.88562870720*t);
+   emb_y_1+=     0.00000000041 * cos( 4.39843263561 +    18052.92954315780*t);
+   emb_y_1+=     0.00000000032 * cos( 4.69085169487 +     6231.79551212899*t);
+   emb_y_1+=     0.00000000032 * cos( 5.09927453137 +     6334.35618785381*t);
+   emb_y_1+=     0.00000000033 * cos( 4.71041781828 +     5017.50837136500*t);
+   emb_y_1+=     0.00000000034 * cos( 3.82387178783 +    22483.84857449259*t);
+   emb_y_1+=     0.00000000029 * cos( 0.69173791628 +    27278.46881644080*t);
+   emb_y_1+=     0.00000000028 * cos( 0.30766379271 +    24080.03501677719*t);
+   emb_y_1+=     0.00000000028 * cos( 1.80200703459 +    21150.81336588360*t);
+   emb_y_1+=     0.00000000034 * cos( 5.49813240003 +     5905.70224207560*t);
+   emb_y_1+=     0.00000000034 * cos( 4.29199382621 +     6660.44945790720*t);
+   emb_y_1+=     0.00000000031 * cos( 0.77685387250 +    15265.88651930040*t);
+   emb_y_1+=     0.00000000038 * cos( 3.55141531291 +    31570.79964939120*t);
+   emb_y_1+=     0.00000000027 * cos( 2.03421614285 +    36949.23080842420*t);
+   emb_y_1+=     0.00000000030 * cos( 0.75326086687 +    12586.92709547520*t);
+   emb_y_1+=     0.00000000027 * cos( 4.99665333874 +    24383.07910844140*t);
+   emb_y_1+=     0.00000000028 * cos( 3.97690125997 +    34911.41207609100*t);
+   emb_y_1+=     0.00000000027 * cos( 4.60109049263 +     1692.16566950240*t);
+   emb_y_1+=     0.00000000027 * cos( 3.01465510733 +    23581.25817731760*t);
+   emb_y_1+=     0.00000000028 * cos( 5.79946172825 +    16193.65917750039*t);
+   emb_y_1+=     0.00000000027 * cos( 5.36876316169 +     6360.74962041940*t);
+   emb_y_1+=     0.00000000027 * cos( 4.42136306455 +     6205.40207956340*t);
+   emb_y_1+=     0.00000000032 * cos( 5.04711925057 +    27511.46787353720*t);
+   emb_y_1+=     0.00000000027 * cos( 0.74966550512 +    32367.09765620760*t);
+   emb_y_1+=     0.00000000027 * cos( 3.34326768159 +     4487.81740627040*t);
+   emb_y_1+=     0.00000000030 * cos( 4.32132101590 +    52175.80628314840*t);
+   emb_y_1+=     0.00000000029 * cos( 4.00870617495 +    46386.99374331120*t);
+   emb_y_1+=     0.00000000026 * cos( 1.30285929374 +     4797.09572892620*t);
+   emb_y_1+=     0.00000000028 * cos( 5.37411984168 +    10239.58386601080*t);
+   emb_y_1+=     0.00000000026 * cos( 5.08151018860 +     6339.97422492700*t);
+   emb_y_1+=     0.00000000026 * cos( 4.70861603764 +     6226.17747505580*t);
+   emb_y_1+=     0.00000000029 * cos( 4.88985723155 +     6592.35417264720*t);
+   emb_y_1+=     0.00000000029 * cos( 4.90026899469 +     5973.79752733560*t);
+   emb_y_1+=     0.00000000026 * cos( 3.80920015468 +    25158.60171976540*t);
+   emb_y_1+=     0.00000000026 * cos( 4.12353757087 +     9910.58332750900*t);
+   emb_y_1+=     0.00000000031 * cos( 2.37034813214 +    24491.42579258340*t);
+   emb_y_1+=     0.00000000024 * cos( 5.26265995716 +    11616.97609101300*t);
+   emb_y_1+=     0.00000000024 * cos( 4.68388215645 +    31172.65064598300*t);
+   emb_y_1+=     0.00000000029 * cos( 4.62741692698 +    19402.79695281660*t);
+   emb_y_1+=     0.00000000025 * cos( 1.21411432675 +     6302.74561089119*t);
+   emb_y_1+=     0.00000000025 * cos( 2.29282659231 +     6263.40608909160*t);
+   emb_y_1+=     0.00000000026 * cos( 4.64443513632 +    30666.15495843280*t);
+   emb_y_1+=     0.00000000027 * cos( 2.79813911349 +    19800.94595622480*t);
+   emb_y_1+=     0.00000000024 * cos( 2.16373217847 +    31968.94865279940*t);
+   emb_y_1+=     0.00000000023 * cos( 5.33583152283 +    17892.93839400359*t);
+   emb_y_1+=     0.00000000023 * cos( 5.05383022333 +    29424.63423291600*t);
+   emb_y_1+=     0.00000000024 * cos( 6.12085409163 +    20597.24396304120*t);
+   emb_y_1+=     0.00000000024 * cos( 1.65155411828 +    23550.34405168259*t);
+   emb_y_1+=     0.00000000023 * cos( 6.07302432574 +     6716.78758786820*t);
+   emb_y_1+=     0.00000000022 * cos( 2.42662159283 +     7445.55055439920*t);
+   emb_y_1+=     0.00000000027 * cos( 1.56823801633 +    28628.33622609960*t);
+   emb_y_1+=     0.00000000022 * cos( 3.39570466574 +    40879.44050464380*t);
+   emb_y_1+=     0.00000000025 * cos( 5.52794317615 +    14155.22459526660*t);
+   emb_y_1+=     0.00000000023 * cos( 1.49204312505 +     6272.03014972750*t);
+   emb_y_1+=     0.00000000023 * cos( 2.01489779401 +     6294.12155025530*t);
+   emb_y_1+=     0.00000000021 * cos( 4.21170718841 +    22779.43724619380*t);
+   emb_y_1+=     0.00000000024 * cos( 3.83904364251 +     8186.51266249260*t);
+   emb_y_1+=     0.00000000021 * cos( 2.51494728863 +    30220.93223973240*t);
+   emb_y_1+=     0.00000000022 * cos( 6.09500922463 +     4379.63903749020*t);
+   emb_y_1+=     0.00000000022 * cos( 2.95144246643 +    28759.81087748319*t);
+   emb_y_1+=     0.00000000020 * cos( 0.07597154617 +    28313.28880466100*t);
+   emb_y_1+=     0.00000000026 * cos( 1.83971959567 +    18422.62935909819*t);
+   emb_y_1+=     0.00000000022 * cos( 1.06284971840 +    27707.54249429480*t);
+   emb_y_1+=     0.00000000027 * cos( 5.38238146057 +    34513.26307268280*t);
+   emb_y_1+=     0.00000000025 * cos( 1.96129048463 +    55798.45835839840*t);
+   emb_y_1+=     0.00000000019 * cos( 0.75805525916 +    11712.95531823080*t);
+   emb_y_1+=     0.00000000025 * cos( 4.06271601240 +    29864.33402730900*t);
+   emb_y_1+=     0.00000000024 * cos( 5.92086856803 +     9411.46461508720*t);
+   emb_y_1+=     0.00000000020 * cos( 0.08212819079 +     6379.05507720920*t);
+   emb_y_1+=     0.00000000020 * cos( 3.42481272827 +     6187.09662277360*t);
+   emb_y_1+=     0.00000000022 * cos( 2.87393551670 +    18326.65013188039*t);
+   emb_y_1+=     0.00000000020 * cos( 4.26784309163 +    21548.96236929180*t);
+   emb_y_1+=     0.00000000022 * cos( 2.89425510955 +     6332.83287546320*t);
+   emb_y_1+=     0.00000000022 * cos( 0.61268580951 +     6233.31882451960*t);
+   emb_y_1+=     0.00000000021 * cos( 5.43088404515 +    50317.20343953080*t);
+   emb_y_1+=     0.00000000020 * cos( 0.38470757997 +    22345.26037610820*t);
+   emb_y_1+=     0.00000000018 * cos( 3.04737429764 +     6593.91664885980*t);
+   emb_y_1+=     0.00000000018 * cos( 0.45956662142 +     5972.23505112300*t);
+   emb_y_1+=     0.00000000019 * cos( 0.79376120003 +     5618.31980486140*t);
+   emb_y_1+=     0.00000000019 * cos( 2.71317971903 +     6947.83189512140*t);
+   emb_y_1+=     0.00000000018 * cos( 0.28017726157 +    34115.11406927460*t);
+   emb_y_1+=     0.00000000020 * cos( 2.60251083721 +    42456.78404709160*t);
+   emb_y_1+=     0.00000000020 * cos( 1.28313240949 +     9070.11887384880*t);
+   emb_y_1=emb_y_1 * t;
+
+   double emb_y_2=0.0;
+   emb_y_2+=     0.00005080208 * cos( 0.00000000000 +        0.00000000000*t);
+   emb_y_2+=     0.00002178016 * cos( 2.82957544235 +    12566.15169998280*t);
+   emb_y_2+=     0.00001020487 * cos( 4.63746718598 +     6283.07584999140*t);
+   emb_y_2+=     0.00000092688 * cos( 2.62218748420 +    18849.22754997420*t);
+   emb_y_2+=     0.00000013680 * cos( 1.80736896884 +     6286.59896834040*t);
+   emb_y_2+=     0.00000013668 * cos( 1.70191859339 +     6279.55273164240*t);
+   emb_y_2+=     0.00000008190 * cos( 2.84848380902 +     6256.77753019160*t);
+   emb_y_2+=     0.00000008145 * cos( 0.65845469208 +     6309.37416979120*t);
+   emb_y_2+=     0.00000008004 * cos( 4.35208516133 +     6127.65545055720*t);
+   emb_y_2+=     0.00000007754 * cos( 5.43893832496 +     6438.49624942560*t);
+   emb_y_2+=     0.00000002895 * cos( 3.37918924673 +     4705.73230754360*t);
+   emb_y_2+=     0.00000003730 * cos( 0.91975073803 +     5507.55323866740*t);
+   emb_y_2+=     0.00000002937 * cos( 2.51081180566 +    25132.30339996560*t);
+   emb_y_2+=     0.00000003203 * cos( 2.57193731532 +     7058.59846131540*t);
+   emb_y_2+=     0.00000002268 * cos( 1.94917686677 +      775.52261132400*t);
+   emb_y_2+=     0.00000002148 * cos( 0.14152116438 +     7860.41939243920*t);
+   emb_y_2+=     0.00000002031 * cos( 0.72283806291 +     6275.96230299060*t);
+   emb_y_2+=     0.00000002031 * cos( 2.78410285615 +     6290.18939699220*t);
+   emb_y_2+=     0.00000001904 * cos( 2.89589054598 +     5486.77784317500*t);
+   emb_y_2+=     0.00000001728 * cos( 3.66132960737 +     1059.38193018920*t);
+   emb_y_2+=     0.00000001553 * cos( 0.56386962806 +     6040.34724601740*t);
+   emb_y_2+=     0.00000001563 * cos( 0.61022183613 +     7079.37385680780*t);
+   emb_y_2+=     0.00000001459 * cos( 1.88257310296 +     5753.38488489680*t);
+   emb_y_2+=     0.00000001425 * cos( 1.59147859762 +     6812.76681508600*t);
+   emb_y_2+=     0.00000001436 * cos( 1.45357168378 +     5884.92684658320*t);
+   emb_y_2+=     0.00000001470 * cos( 2.94579999718 +     6525.80445396540*t);
+   emb_y_2+=     0.00000001288 * cos( 2.05812819149 +     6681.22485339960*t);
+   emb_y_2+=     0.00000001338 * cos( 5.75684444528 +     5729.50644714900*t);
+   emb_y_2+=     0.00000001039 * cos( 5.50782141440 +     6284.05617105960*t);
+   emb_y_2+=     0.00000001039 * cos( 4.28230481184 +     6282.09552892320*t);
+   emb_y_2+=     0.00000001191 * cos( 4.01917074242 +     6836.64525283380*t);
+   emb_y_2+=     0.00000000904 * cos( 0.86455323243 +     5331.35744374080*t);
+   emb_y_2+=     0.00000000871 * cos( 2.27688978691 +     4933.20844033260*t);
+   emb_y_2+=     0.00000000877 * cos( 1.49469829015 +    11790.62908865880*t);
+   emb_y_2+=     0.00000000774 * cos( 2.68163901032 +     7234.79425624200*t);
+   emb_y_2+=     0.00000000746 * cos( 5.06505548485 +     4136.91043351620*t);
+   emb_y_2+=     0.00000000653 * cos( 6.14159854835 +    11506.76976979360*t);
+   emb_y_2+=     0.00000000653 * cos( 1.20042991173 +     7632.94325965020*t);
+   emb_y_2+=     0.00000000855 * cos( 3.67241064264 +     4535.05943692440*t);
+   emb_y_2+=     0.00000000664 * cos( 3.21978126910 +     1194.44701022460*t);
+   emb_y_2+=     0.00000000814 * cos( 1.07520259172 +      155.42039943420*t);
+   emb_y_2+=     0.00000000579 * cos( 4.64690913221 +      796.29800681640*t);
+   emb_y_2+=     0.00000000778 * cos( 4.32748087995 +     5088.62883976680*t);
+   emb_y_2+=     0.00000000511 * cos( 0.71257437268 +     4694.00295470760*t);
+   emb_y_2+=     0.00000000686 * cos( 3.47131000180 +      801.82093112380*t);
+   emb_y_2+=     0.00000000677 * cos( 1.14891148781 +     3128.38876509580*t);
+   emb_y_2+=     0.00000000499 * cos( 1.45672603171 +    12569.67481833180*t);
+   emb_y_2+=     0.00000000480 * cos( 4.72566522648 +     8429.24126646660*t);
+   emb_y_2+=     0.00000000598 * cos( 6.12701256696 +     8031.09226305840*t);
+   emb_y_2+=     0.00000000499 * cos( 4.61877818070 +     1589.07289528380*t);
+   emb_y_2+=     0.00000000562 * cos( 5.48344520380 +     7477.52286021600*t);
+   emb_y_2+=     0.00000000415 * cos( 5.06370124318 +     6496.37494542940*t);
+   emb_y_2+=     0.00000000414 * cos( 4.73489535739 +     6069.77675455340*t);
+   emb_y_2+=     0.00000000404 * cos( 0.15585991312 +     3738.76143010800*t);
+   emb_y_2+=     0.00000000500 * cos( 1.83746339810 +     1592.59601363280*t);
+   emb_y_2+=     0.00000000392 * cos( 0.64094689719 +       26.29831979980*t);
+   emb_y_2+=     0.00000000366 * cos( 5.93255688855 +      553.56940284240*t);
+   emb_y_2+=     0.00000000320 * cos( 0.43684109891 +     1990.74501704100*t);
+   emb_y_2+=     0.00000000318 * cos( 4.53872367492 +      951.71840625060*t);
+   emb_y_2+=     0.00000000361 * cos( 3.78139914914 +     7084.89678111520*t);
+   emb_y_2+=     0.00000000356 * cos( 4.78764157089 +    12721.57209941700*t);
+   emb_y_2+=     0.00000000279 * cos( 5.51531767271 +     4690.47983635860*t);
+   emb_y_2+=     0.00000000347 * cos( 6.05487619509 +     5481.25491886760*t);
+   emb_y_2+=     0.00000000279 * cos( 0.65072513623 +      242.72860397400*t);
+   emb_y_2+=     0.00000000367 * cos( 2.35508391326 +     9437.76293488700*t);
+   emb_y_2+=     0.00000000260 * cos( 0.98725902772 +     5223.69391980220*t);
+   emb_y_2+=     0.00000000255 * cos( 0.26797013701 +    11371.70468975820*t);
+   emb_y_2+=     0.00000000242 * cos( 3.35248006615 +     8827.39026987480*t);
+   emb_y_2+=     0.00000000245 * cos( 2.60270196816 +     3154.68708489560*t);
+   emb_y_2+=     0.00000000222 * cos( 0.17857635292 +     4164.31198961300*t);
+   emb_y_2+=     0.00000000221 * cos( 1.83683124614 +     1748.01641306700*t);
+   emb_y_2+=     0.00000000288 * cos( 3.16521452041 +     1349.86740965880*t);
+   emb_y_2+=     0.00000000204 * cos( 5.15479639183 +    11769.85369316640*t);
+   emb_y_2+=     0.00000000188 * cos( 4.08810200628 +     7875.67186362420*t);
+   emb_y_2+=     0.00000000204 * cos( 5.17246322251 +    10977.07880469900*t);
+   emb_y_2+=     0.00000000188 * cos( 5.32022679153 +     2388.89402044920*t);
+   emb_y_2+=     0.00000000183 * cos( 2.53020407005 +     7342.45778018060*t);
+   emb_y_2+=     0.00000000196 * cos( 4.81306031436 +     3930.20969621960*t);
+   emb_y_2+=     0.00000000207 * cos( 0.71423535949 +    13367.97263110660*t);
+   emb_y_2+=     0.00000000206 * cos( 1.68118727655 +    10973.55568635000*t);
+   emb_y_2+=     0.00000000144 * cos( 3.12119544517 +    10575.40668294180*t);
+   emb_y_2+=     0.00000000155 * cos( 4.24333182826 +    12410.73130054860*t);
+   emb_y_2+=     0.00000000135 * cos( 5.62460495351 +    12592.45001978260*t);
+   emb_y_2+=     0.00000000142 * cos( 3.43769023959 +      529.69096509460*t);
+   emb_y_2+=     0.00000000125 * cos( 2.30734650656 +    12562.62858163380*t);
+   emb_y_2+=     0.00000000118 * cos( 3.15911802638 +    13119.72110282519*t);
+   emb_y_2+=     0.00000000124 * cos( 4.88716980692 +     8635.94200376320*t);
+   emb_y_2+=     0.00000000111 * cos( 0.80973565731 +      398.14900340820*t);
+   emb_y_2+=     0.00000000099 * cos( 1.74904579789 +    13517.87010623340*t);
+   emb_y_2+=     0.00000000097 * cos( 4.08519883679 +    17260.15465469040*t);
+   emb_y_2+=     0.00000000102 * cos( 2.80042792551 +    10988.80815753500*t);
+   emb_y_2+=     0.00000000093 * cos( 4.87729654087 +     1577.34354244780*t);
+   emb_y_2+=     0.00000000094 * cos( 4.56913615419 +    10177.25767953360*t);
+   emb_y_2+=     0.00000000095 * cos( 0.58410858338 +     2146.16541647520*t);
+   emb_y_2+=     0.00000000100 * cos( 0.77889460323 +    13095.84266507740*t);
+   emb_y_2+=     0.00000000082 * cos( 3.65082093235 +    16730.46368959580*t);
+   emb_y_2+=     0.00000000109 * cos( 2.42647120751 +    12539.85338018300*t);
+   emb_y_2+=     0.00000000080 * cos( 2.67636202935 +    12323.42309600880*t);
+   emb_y_2+=     0.00000000105 * cos( 1.96765666886 +     5642.19824260920*t);
+   emb_y_2+=     0.00000000101 * cos( 0.26697958405 +     4292.33083295040*t);
+   emb_y_2+=     0.00000000075 * cos( 4.15987686546 +     6709.67404086740*t);
+   emb_y_2+=     0.00000000075 * cos( 5.63024936078 +     5856.47765911540*t);
+   emb_y_2+=     0.00000000071 * cos( 5.31207567706 +    14314.16811304980*t);
+   emb_y_2+=     0.00000000092 * cos( 0.41159879070 +    13916.01910964160*t);
+   emb_y_2+=     0.00000000091 * cos( 1.53881053794 +     6923.95345737360*t);
+   emb_y_2+=     0.00000000082 * cos( 2.43138857903 +    31415.37924995700*t);
+   emb_y_2+=     0.00000000061 * cos( 2.66905862210 +     8273.82086703240*t);
+   emb_y_2+=     0.00000000055 * cos( 2.98440193981 +     5643.17856367740*t);
+   emb_y_2+=     0.00000000056 * cos( 3.14343584528 +      639.89728631400*t);
+   emb_y_2+=     0.00000000054 * cos( 4.06071028317 +    16496.36139620240*t);
+   emb_y_2+=     0.00000000055 * cos( 2.26469062405 +    13341.67431130680*t);
+   emb_y_2+=     0.00000000052 * cos( 5.97085689483 +     3634.62102451840*t);
+   emb_y_2+=     0.00000000050 * cos( 2.61636629528 +    12573.26524698360*t);
+   emb_y_2+=     0.00000000052 * cos( 6.19672300540 +    15720.83878487840*t);
+   emb_y_2+=     0.00000000065 * cos( 4.92633618792 +     4732.03062734340*t);
+   emb_y_2+=     0.00000000066 * cos( 3.92505637420 +     2787.04302385740*t);
+   emb_y_2+=     0.00000000048 * cos( 0.45273646609 +    13362.44970679920*t);
+   emb_y_2+=     0.00000000053 * cos( 2.55293773264 +    14919.01785375460*t);
+   emb_y_2+=     0.00000000047 * cos( 2.25879651839 +    12964.30070339100*t);
+   emb_y_2+=     0.00000000047 * cos( 1.06895129610 +     3894.18182954220*t);
+   emb_y_2+=     0.00000000042 * cos( 5.09300549015 +    24072.92146977640*t);
+   emb_y_2+=     0.00000000051 * cos( 5.97812516445 +     6195.76764545159*t);
+   emb_y_2+=     0.00000000051 * cos( 3.81200106179 +     6370.38405453120*t);
+   emb_y_2+=     0.00000000048 * cos( 5.31227829697 +     2118.76386037840*t);
+   emb_y_2+=     0.00000000041 * cos( 0.92655469135 +    18319.53658487960*t);
+   emb_y_2+=     0.00000000041 * cos( 2.52247946610 +     3185.19202726560*t);
+   emb_y_2+=     0.00000000039 * cos( 2.70559904392 +    20426.57109242200*t);
+   emb_y_2+=     0.00000000045 * cos( 6.04609930388 +     9779.10867612540*t);
+   emb_y_2+=     0.00000000044 * cos( 3.24184126256 +     8662.24032356300*t);
+   emb_y_2+=     0.00000000033 * cos( 1.67764931937 +     3496.03282613400*t);
+   emb_y_2+=     0.00000000035 * cos( 5.17885502556 +    18875.52586977400*t);
+   emb_y_2+=     0.00000000035 * cos( 4.51995137508 +    10419.98628350760*t);
+   emb_y_2+=     0.00000000031 * cos( 5.35886426810 +    23013.53953958720*t);
+   emb_y_2+=     0.00000000030 * cos( 2.82636067861 +    12808.88030395680*t);
+   emb_y_2+=     0.00000000031 * cos( 3.95971884541 +    14712.31711645800*t);
+   emb_y_2+=     0.00000000040 * cos( 4.02926744005 +    17789.84561978500*t);
+   emb_y_2+=     0.00000000031 * cos( 5.47911308336 +    22779.43724619380*t);
+   emb_y_2+=     0.00000000030 * cos( 2.31682203487 +    10447.38783960440*t);
+   emb_y_2+=     0.00000000028 * cos( 2.11876435927 +    14143.49524243060*t);
+   emb_y_2+=     0.00000000028 * cos( 3.91478954977 +     5429.87946823940*t);
+   emb_y_2+=     0.00000000030 * cos( 1.17555913144 +     9380.95967271720*t);
+   emb_y_2+=     0.00000000026 * cos( 0.72314347990 +     9623.68827669120*t);
+   emb_y_2+=     0.00000000026 * cos( 2.78379743916 +     2942.46342329160*t);
+   emb_y_2+=     0.00000000030 * cos( 1.83762487734 +    10213.28554621100*t);
+   emb_y_2+=     0.00000000026 * cos( 0.18151529861 +    18209.33026366019*t);
+   emb_y_2+=     0.00000000025 * cos( 1.11733039286 +     3583.34103067380*t);
+   emb_y_2+=     0.00000000024 * cos( 1.36775216012 +     9917.69687450980*t);
+   emb_y_2+=     0.00000000026 * cos( 5.61732376806 +    12012.58229714040*t);
+   emb_y_2+=     0.00000000023 * cos( 0.92088373151 +    12559.03815298200*t);
+   emb_y_2+=     0.00000000023 * cos( 5.52325028212 +    11015.10647733480*t);
+   emb_y_2+=     0.00000000023 * cos( 0.69647593831 +    26709.64694241340*t);
+   emb_y_2+=     0.00000000021 * cos( 0.81252266714 +    17256.63153634140*t);
+   emb_y_2+=     0.00000000019 * cos( 0.35979025769 +    11926.25441366880*t);
+   emb_y_2+=     0.00000000018 * cos( 0.71176167314 +    11614.43329373220*t);
+   emb_y_2+=     0.00000000019 * cos( 4.89783197147 +    18073.70493865020*t);
+   emb_y_2+=     0.00000000019 * cos( 2.09021029647 +    30639.85663863300*t);
+   emb_y_2+=     0.00000000022 * cos( 2.19299415984 +    17298.18232732620*t);
+   emb_y_2+=     0.00000000019 * cos( 1.23251732918 +    19651.04848109800*t);
+   emb_y_2+=     0.00000000017 * cos( 2.38624204004 +     8982.81066930900*t);
+   emb_y_2+=     0.00000000021 * cos( 2.20285403235 +    16858.48253293320*t);
+   emb_y_2+=     0.00000000021 * cos( 2.89968186701 +    12036.46073488820*t);
+   emb_y_2+=     0.00000000014 * cos( 3.30566599775 +    16460.33352952499*t);
+   emb_y_2+=     0.00000000015 * cos( 2.64766821618 +    20199.09495963300*t);
+   emb_y_2+=     0.00000000014 * cos( 4.87181144278 +    22483.84857449259*t);
+   emb_y_2+=     0.00000000019 * cos( 1.24341688060 +     8671.96987044060*t);
+   emb_y_2+=     0.00000000016 * cos( 5.12394281212 +    23543.23050468179*t);
+   emb_y_2+=     0.00000000012 * cos( 1.13908210541 +    17996.03116822220*t);
+   emb_y_2+=     0.00000000012 * cos( 5.43458827343 +    19402.79695281660*t);
+   emb_y_2+=     0.00000000012 * cos( 2.80798050151 +    10818.13528691580*t);
+   emb_y_2+=     0.00000000015 * cos( 5.28864619208 +    12567.13202105100*t);
+   emb_y_2+=     0.00000000011 * cos( 0.21390152319 +    28766.92442448400*t);
+   emb_y_2+=     0.00000000011 * cos( 0.68624711981 +    14945.31617355440*t);
+   emb_y_2+=     0.00000000011 * cos( 3.36248263984 +    18451.07854656599*t);
+   emb_y_2+=     0.00000000011 * cos( 4.69826436697 +     2352.86615377180*t);
+   emb_y_2+=     0.00000000011 * cos( 0.30993049479 +    22805.73556599360*t);
+   emb_y_2+=     0.00000000011 * cos( 0.67524830165 +    35371.88726597640*t);
+   emb_y_2+=     0.00000000014 * cos( 0.36237106347 +    21228.39202354580*t);
+   emb_y_2+=     0.00000000010 * cos( 4.48227917427 +    23141.55838292460*t);
+   emb_y_2+=     0.00000000010 * cos( 0.98562604234 +    22345.26037610820*t);
+   emb_y_2+=     0.00000000009 * cos( 6.20387541356 +    29826.30635467320*t);
+   emb_y_2+=     0.00000000010 * cos( 4.89597714031 +    12779.45079542080*t);
+   emb_y_2+=     0.00000000010 * cos( 4.53062887484 +      213.29909543800*t);
+   emb_y_2+=     0.00000000011 * cos( 5.10892321991 +    25158.60171976540*t);
+   emb_y_2+=     0.00000000009 * cos( 2.89659186256 +    23581.25817731760*t);
+   emb_y_2=emb_y_2 * t * t;
+
+   double emb_y_3=0.0;
+   emb_y_3+=     0.00000128116 * cos( 3.14159265359 +        0.00000000000*t);
+   emb_y_3+=     0.00000072366 * cos( 1.32019314413 +    12566.15169998280*t);
+   emb_y_3+=     0.00000013847 * cos( 5.37444701633 +     6283.07584999140*t);
+   emb_y_3+=     0.00000005044 * cos( 0.96001931081 +    18849.22754997420*t);
+   emb_y_3+=     0.00000001481 * cos( 0.66592602555 +     6438.49624942560*t);
+   emb_y_3+=     0.00000001481 * cos( 2.84101489351 +     6127.65545055720*t);
+   emb_y_3+=     0.00000000644 * cos( 3.31488272100 +     6279.55273164240*t);
+   emb_y_3+=     0.00000000644 * cos( 0.19205819806 +     6286.59896834040*t);
+   emb_y_3+=     0.00000000318 * cos( 2.06760677503 +     6040.34724601740*t);
+   emb_y_3+=     0.00000000318 * cos( 1.43933414403 +     6525.80445396540*t);
+   emb_y_3+=     0.00000000220 * cos( 0.77087312633 +    25132.30339996560*t);
+   emb_y_3+=     0.00000000201 * cos( 5.53953658848 +     6836.64525283380*t);
+   emb_y_3+=     0.00000000201 * cos( 4.25058963776 +     5729.50644714900*t);
+   emb_y_3+=     0.00000000036 * cos( 6.12441668493 +    12569.67481833180*t);
+   emb_y_3+=     0.00000000033 * cos( 2.66474952932 +    12410.73130054860*t);
+   emb_y_3+=     0.00000000018 * cos( 0.75629855075 +    12721.57209941700*t);
+   emb_y_3+=     0.00000000011 * cos( 4.70234349211 +    12592.45001978260*t);
+   emb_y_3=emb_y_3 * t * t * t;
+
+   double emb_y_4=0.0;
+   emb_y_4+=     0.00000004187 * cos( 3.14159265359 +        0.00000000000*t);
+   emb_y_4+=     0.00000001931 * cos( 5.89883838000 +    12566.15169998280*t);
+   emb_y_4+=     0.00000000595 * cos( 0.51679590840 +     6283.07584999140*t);
+   emb_y_4+=     0.00000000210 * cos( 2.17239897368 +     6438.49624942560*t);
+   emb_y_4+=     0.00000000210 * cos( 1.33454194538 +     6127.65545055720*t);
+   emb_y_4+=     0.00000000229 * cos( 5.59311519226 +    18849.22754997420*t);
+   emb_y_4+=     0.00000000021 * cos( 4.89633014340 +     6279.55273164240*t);
+   emb_y_4+=     0.00000000021 * cos( 4.89379608284 +     6286.59896834040*t);
+   emb_y_4+=     0.00000000011 * cos( 5.36791829228 +    25132.30339996560*t);
+   emb_y_4+=     0.00000000005 * cos( 1.11835906029 +    12410.73130054860*t);
+   emb_y_4=emb_y_4 * t * t * t * t;
+
+   double emb_y_5=0.0;
+   emb_y_5+=     0.00000000072 * cos( 0.00000000000 +        0.00000000000*t);
+   emb_y_5+=     0.00000000043 * cos( 4.51710019379 +    12566.15169998280*t);
+   emb_y_5+=     0.00000000025 * cos( 3.76700032666 +     6438.49624942560*t);
+   emb_y_5+=     0.00000000025 * cos( 6.02312589958 +     6127.65545055720*t);
+   emb_y_5+=     0.00000000021 * cos( 2.56882591004 +     6283.07584999140*t);
+   emb_y_5+=     0.00000000008 * cos( 3.95508428463 +    18849.22754997420*t);
+   emb_y_5=emb_y_5 * t * t * t * t * t;
+
+   return emb_y_0+emb_y_1+emb_y_2+emb_y_3+emb_y_4+emb_y_5;
+}
+
+double vsop87a_full::emb_z(double t){
+   double emb_z_0=0.0;
+   emb_z_0+=     0.00000101625 * cos( 5.42248110597 +     5507.55323866740*t);
+   emb_z_0+=     0.00000080461 * cos( 3.88027157914 +     5223.69391980220*t);
+   emb_z_0+=     0.00000043855 * cos( 3.70430347723 +     2352.86615377180*t);
+   emb_z_0+=     0.00000031853 * cos( 3.99997840986 +     1577.34354244780*t);
+   emb_z_0+=     0.00000016395 * cos( 3.56456162523 +     5856.47765911540*t);
+   emb_z_0+=     0.00000018150 * cos( 4.98479613321 +     6283.07584999140*t);
+   emb_z_0+=     0.00000014477 * cos( 3.70258423465 +     9437.76293488700*t);
+   emb_z_0+=     0.00000014325 * cos( 3.41020246931 +    10213.28554621100*t);
+   emb_z_0+=     0.00000011204 * cos( 4.83021499098 +    14143.49524243060*t);
+   emb_z_0+=     0.00000010896 * cos( 2.08023708321 +     6812.76681508600*t);
+   emb_z_0+=     0.00000009716 * cos( 3.47560205064 +     4694.00295470760*t);
+   emb_z_0+=     0.00000008792 * cos( 4.44890729280 +     5753.38488489680*t);
+   emb_z_0+=     0.00000008337 * cos( 4.99167706048 +     7084.89678111520*t);
+   emb_z_0+=     0.00000006921 * cos( 4.32559054073 +     6275.96230299060*t);
+   emb_z_0+=     0.00000009257 * cos( 1.14182647820 +     6620.89011318780*t);
+   emb_z_0+=     0.00000007205 * cos( 3.62441599378 +      529.69096509460*t);
+   emb_z_0+=     0.00000005186 * cos( 6.25384533126 +    18073.70493865020*t);
+   emb_z_0+=     0.00000005027 * cos( 2.49727910749 +     4705.73230754360*t);
+   emb_z_0+=     0.00000004529 * cos( 2.33827747356 +     6309.37416979120*t);
+   emb_z_0+=     0.00000004754 * cos( 0.71100829534 +     5884.92684658320*t);
+   emb_z_0+=     0.00000004292 * cos( 1.10034019036 +     6681.22485339960*t);
+   emb_z_0+=     0.00000003855 * cos( 1.82331477760 +     5486.77784317500*t);
+   emb_z_0+=     0.00000005256 * cos( 4.42445744523 +     7860.41939243920*t);
+   emb_z_0+=     0.00000004025 * cos( 5.11990285618 +    13367.97263110660*t);
+   emb_z_0+=     0.00000004061 * cos( 6.02923989849 +     3930.20969621960*t);
+   emb_z_0+=     0.00000003798 * cos( 0.44370219832 +     3154.68708489560*t);
+   emb_z_0+=     0.00000002933 * cos( 5.12428135363 +     1059.38193018920*t);
+   emb_z_0+=     0.00000003424 * cos( 5.42065191700 +     6069.77675455340*t);
+   emb_z_0+=     0.00000002912 * cos( 1.92688087702 +    10977.07880469900*t);
+   emb_z_0+=     0.00000002506 * cos( 0.60484952637 +     6496.37494542940*t);
+   emb_z_0+=     0.00000002680 * cos( 1.39497359287 +    22003.91463486980*t);
+   emb_z_0+=     0.00000002120 * cos( 4.30691000285 +     5643.17856367740*t);
+   emb_z_0+=     0.00000002235 * cos( 0.81363184041 +     8635.94200376320*t);
+   emb_z_0+=     0.00000001819 * cos( 3.75748003107 +     3340.61242669980*t);
+   emb_z_0+=     0.00000002230 * cos( 2.76407822446 +    12036.46073488820*t);
+   emb_z_0+=     0.00000001945 * cos( 5.69949789177 +    11790.62908865880*t);
+   emb_z_0+=     0.00000001529 * cos( 1.98639348122 +      398.14900340820*t);
+   emb_z_0+=     0.00000001581 * cos( 3.19976230948 +     5088.62883976680*t);
+   emb_z_0+=     0.00000001424 * cos( 6.26116472313 +     2544.31441988340*t);
+   emb_z_0+=     0.00000001405 * cos( 4.69394873481 +     7058.59846131540*t);
+   emb_z_0+=     0.00000001588 * cos( 0.25704624784 +    17298.18232732620*t);
+   emb_z_0+=     0.00000001478 * cos( 2.81808207569 +    25934.12433108940*t);
+   emb_z_0+=     0.00000001123 * cos( 2.38605285936 +     3738.76143010800*t);
+   emb_z_0+=     0.00000000990 * cos( 4.29600699156 +     9225.53927328300*t);
+   emb_z_0+=     0.00000000923 * cos( 3.07195736431 +     4164.31198961300*t);
+   emb_z_0+=     0.00000000863 * cos( 0.55119493097 +     8429.24126646660*t);
+   emb_z_0+=     0.00000000984 * cos( 5.96784225392 +     7079.37385680780*t);
+   emb_z_0+=     0.00000001086 * cos( 3.14159265359 +        0.00000000000*t);
+   emb_z_0+=     0.00000000824 * cos( 1.50050055147 +    10447.38783960440*t);
+   emb_z_0+=     0.00000000829 * cos( 3.51336978167 +    11506.76976979360*t);
+   emb_z_0+=     0.00000000908 * cos( 0.12102885474 +    11015.10647733480*t);
+   emb_z_0+=     0.00000000742 * cos( 1.99159139281 +    26087.90314157420*t);
+   emb_z_0+=     0.00000000850 * cos( 4.24120016095 +    29864.33402730900*t);
+   emb_z_0+=     0.00000000762 * cos( 2.90028034892 +     4732.03062734340*t);
+   emb_z_0+=     0.00000000711 * cos( 1.38596151724 +     2146.16541647520*t);
+   emb_z_0+=     0.00000000766 * cos( 0.61397570576 +      796.29800681640*t);
+   emb_z_0+=     0.00000000708 * cos( 1.91406542362 +     8031.09226305840*t);
+   emb_z_0+=     0.00000000810 * cos( 5.13961498599 +     2942.46342329160*t);
+   emb_z_0+=     0.00000000751 * cos( 1.67479850166 +    21228.39202354580*t);
+   emb_z_0+=     0.00000000589 * cos( 2.02193316047 +      775.52261132400*t);
+   emb_z_0+=     0.00000000513 * cos( 2.60734651128 +    12566.15169998280*t);
+   emb_z_0+=     0.00000000550 * cos( 1.61119298963 +      801.82093112380*t);
+   emb_z_0+=     0.00000000592 * cos( 4.59481504319 +     4690.47983635860*t);
+   emb_z_0+=     0.00000000519 * cos( 5.72925039114 +     8827.39026987480*t);
+   emb_z_0+=     0.00000000503 * cos( 5.66433137112 +    33794.54372352860*t);
+   emb_z_0+=     0.00000000476 * cos( 3.69640472448 +      213.29909543800*t);
+   emb_z_0+=     0.00000000463 * cos( 1.24987240483 +    15720.83878487840*t);
+   emb_z_0+=     0.00000000411 * cos( 0.10769444525 +     3128.38876509580*t);
+   emb_z_0+=     0.00000000504 * cos( 3.26207669160 +     7632.94325965020*t);
+   emb_z_0+=     0.00000000394 * cos( 5.81496021156 +      426.59819087600*t);
+   emb_z_0+=     0.00000000476 * cos( 3.52810472421 +    16496.36139620240*t);
+   emb_z_0+=     0.00000000503 * cos( 4.85802444134 +     6290.18939699220*t);
+   emb_z_0+=     0.00000000417 * cos( 0.81920713533 +     5216.58037280140*t);
+   emb_z_0+=     0.00000000363 * cos( 5.70015720452 +     1589.07289528380*t);
+   emb_z_0+=     0.00000000363 * cos( 1.28376436579 +     6206.80977871580*t);
+   emb_z_0+=     0.00000000353 * cos( 4.70059133110 +     7234.79425624200*t);
+   emb_z_0+=     0.00000000381 * cos( 2.61890749829 +     7342.45778018060*t);
+   emb_z_0+=     0.00000000420 * cos( 0.97821135132 +     4136.91043351620*t);
+   emb_z_0+=     0.00000000335 * cos( 6.22057967320 +    12168.00269657460*t);
+   emb_z_0+=     0.00000000387 * cos( 3.09145061418 +    25158.60171976540*t);
+   emb_z_0+=     0.00000000349 * cos( 2.94840720073 +     9623.68827669120*t);
+   emb_z_0+=     0.00000000418 * cos( 3.75759994446 +     5230.80746680300*t);
+   emb_z_0+=     0.00000000396 * cos( 1.22507712354 +     6438.49624942560*t);
+   emb_z_0+=     0.00000000322 * cos( 1.21162178805 +     8662.24032356300*t);
+   emb_z_0+=     0.00000000323 * cos( 5.50254808471 +     1194.44701022460*t);
+   emb_z_0+=     0.00000000372 * cos( 1.72256107938 +    14945.31617355440*t);
+   emb_z_0+=     0.00000000305 * cos( 0.80429352049 +    37724.75341974820*t);
+   emb_z_0+=     0.00000000265 * cos( 6.10358507671 +     6836.64525283380*t);
+   emb_z_0+=     0.00000000250 * cos( 4.56452895547 +     7477.52286021600*t);
+   emb_z_0+=     0.00000000266 * cos( 2.62926282354 +     7238.67559160000*t);
+   emb_z_0+=     0.00000000262 * cos( 1.34297269105 +    11769.85369316640*t);
+   emb_z_0+=     0.00000000263 * cos( 6.22089501237 +     6133.51265285680*t);
+   emb_z_0+=     0.00000000302 * cos( 2.81139045072 +     1748.01641306700*t);
+   emb_z_0+=     0.00000000226 * cos( 2.42346415873 +    11371.70468975820*t);
+   emb_z_0+=     0.00000000216 * cos( 3.68721275185 +     5849.36411211460*t);
+   emb_z_0+=     0.00000000276 * cos( 2.33399767923 +    19651.04848109800*t);
+   emb_z_0+=     0.00000000230 * cos( 0.36165162947 +     5863.59120611620*t);
+   emb_z_0+=     0.00000000200 * cos( 5.86073159059 +     4535.05943692440*t);
+   emb_z_0+=     0.00000000199 * cos( 5.05186622555 +     5429.87946823940*t);
+   emb_z_0+=     0.00000000206 * cos( 3.70556027982 +    10973.55568635000*t);
+   emb_z_0+=     0.00000000210 * cos( 4.50691909144 +    29088.81141598500*t);
+   emb_z_0+=     0.00000000181 * cos( 6.00294783127 +     4292.33083295040*t);
+   emb_z_0+=     0.00000000167 * cos( 6.27348370813 +    10988.80815753500*t);
+   emb_z_0+=     0.00000000183 * cos( 3.19836584401 +    16730.46368959580*t);
+   emb_z_0+=     0.00000000232 * cos( 4.10074982116 +      337.81426319640*t);
+   emb_z_0+=     0.00000000225 * cos( 3.18339652605 +    18875.52586977400*t);
+   emb_z_0+=     0.00000000188 * cos( 2.22746128596 +    41654.96311596780*t);
+   emb_z_0+=     0.00000000160 * cos( 5.20207795189 +     5481.25491886760*t);
+   emb_z_0+=     0.00000000192 * cos( 3.89678943865 +    17789.84561978500*t);
+   emb_z_0+=     0.00000000159 * cos( 5.32590846475 +        7.11354700080*t);
+   emb_z_0+=     0.00000000184 * cos( 3.75465072369 +    23581.25817731760*t);
+   emb_z_0+=     0.00000000144 * cos( 5.19486374726 +    20426.57109242200*t);
+   emb_z_0+=     0.00000000126 * cos( 1.69446958184 +    15110.46611986620*t);
+   emb_z_0+=     0.00000000149 * cos( 1.55306832217 +    10021.83728009940*t);
+   emb_z_0+=     0.00000000127 * cos( 5.31068233515 +      639.89728631400*t);
+   emb_z_0+=     0.00000000118 * cos( 2.73268114062 +    18849.22754997420*t);
+   emb_z_0+=     0.00000000123 * cos( 2.55361087146 +    12592.45001978260*t);
+   emb_z_0+=     0.00000000113 * cos( 6.07178776726 +     6709.67404086740*t);
+   emb_z_0+=     0.00000000110 * cos( 3.51203379263 +    18052.92954315780*t);
+   emb_z_0+=     0.00000000147 * cos( 4.63371971408 +    22805.73556599360*t);
+   emb_z_0+=     0.00000000113 * cos( 1.20711498905 +    14919.01785375460*t);
+   emb_z_0+=     0.00000000108 * cos( 3.76359880753 +    14314.16811304980*t);
+   emb_z_0+=     0.00000000121 * cos( 3.13748782783 +    14712.31711645800*t);
+   emb_z_0+=     0.00000000119 * cos( 5.92110458985 +    33019.02111220460*t);
+   emb_z_0+=     0.00000000124 * cos( 2.93638421840 +    32217.20018108080*t);
+   emb_z_0+=     0.00000000117 * cos( 3.65065271640 +    45585.17281218740*t);
+   emb_z_0+=     0.00000000127 * cos( 4.74596574209 +    49515.38250840700*t);
+   emb_z_0+=     0.00000000120 * cos( 1.04211499785 +     6915.85958930460*t);
+   emb_z_0+=     0.00000000120 * cos( 5.60638811846 +     5650.29211067820*t);
+   emb_z_0+=     0.00000000127 * cos( 2.70766997307 +    12352.85260454480*t);
+   emb_z_0+=     0.00000000097 * cos( 1.07908724794 +     9917.69687450980*t);
+   emb_z_0+=     0.00000000123 * cos( 5.15992792934 +    27511.46787353720*t);
+   emb_z_0+=     0.00000000099 * cos( 4.45774681732 +     4933.20844033260*t);
+   emb_z_0+=     0.00000000087 * cos( 3.93637812950 +    17654.78053974960*t);
+   emb_z_0+=     0.00000000102 * cos( 3.97386522171 +    10818.13528691580*t);
+   emb_z_0+=     0.00000000084 * cos( 4.21241958885 +    22779.43724619380*t);
+   emb_z_0+=     0.00000000103 * cos( 4.63519841733 +     1349.86740965880*t);
+   emb_z_0+=     0.00000000081 * cos( 4.01655240430 +     1592.59601363280*t);
+   emb_z_0+=     0.00000000094 * cos( 4.33873972800 +    36147.40987730040*t);
+   emb_z_0+=     0.00000000076 * cos( 2.83897554420 +    17260.15465469040*t);
+   emb_z_0+=     0.00000000100 * cos( 6.07733097102 +    26735.94526221320*t);
+   emb_z_0+=     0.00000000078 * cos( 0.72462427283 +    12779.45079542080*t);
+   emb_z_0+=     0.00000000086 * cos( 5.26029638250 +    28313.28880466100*t);
+   emb_z_0+=     0.00000000093 * cos( 4.31900620254 +    44809.65020086340*t);
+   emb_z_0+=     0.00000000076 * cos( 6.22743405935 +    13521.75144159140*t);
+   emb_z_0+=     0.00000000081 * cos( 2.16420552322 +    13095.84266507740*t);
+   emb_z_0+=     0.00000000100 * cos( 1.38002787119 +    28286.99048486120*t);
+   emb_z_0+=     0.00000000072 * cos( 1.55820597747 +     6256.77753019160*t);
+   emb_z_0+=     0.00000000082 * cos( 4.95202664555 +    10575.40668294180*t);
+   emb_z_0+=     0.00000000082 * cos( 1.69647647075 +     1990.74501704100*t);
+   emb_z_0+=     0.00000000096 * cos( 6.16038106485 +    24356.78078864160*t);
+   emb_z_0+=     0.00000000075 * cos( 2.29836095644 +     3634.62102451840*t);
+   emb_z_0+=     0.00000000075 * cos( 2.66367876557 +    16200.77272450120*t);
+   emb_z_0+=     0.00000000087 * cos( 0.26630214764 +    31441.67756975680*t);
+
+   double emb_z_1=0.0;
+   emb_z_1+=     0.00227822442 * cos( 3.41372504278 +     6283.07584999140*t);
+   emb_z_1+=     0.00005429282 * cos( 0.00000000000 +        0.00000000000*t);
+   emb_z_1+=     0.00001903183 * cos( 3.37061270964 +    12566.15169998280*t);
+   emb_z_1+=     0.00000023859 * cos( 3.32836261978 +    18849.22754997420*t);
+   emb_z_1+=     0.00000008671 * cos( 1.82289925035 +     5223.69391980220*t);
+   emb_z_1+=     0.00000007767 * cos( 3.88787435832 +     5507.55323866740*t);
+   emb_z_1+=     0.00000007063 * cos( 5.19303169379 +     2352.86615377180*t);
+   emb_z_1+=     0.00000005362 * cos( 2.22429221689 +      529.69096509460*t);
+   emb_z_1+=     0.00000003894 * cos( 2.15568517178 +     6279.55273164240*t);
+   emb_z_1+=     0.00000003892 * cos( 1.53021064904 +     6286.59896834040*t);
+   emb_z_1+=     0.00000003285 * cos( 4.88047587852 +    10213.28554621100*t);
+   emb_z_1+=     0.00000002165 * cos( 6.23220223547 +     7860.41939243920*t);
+   emb_z_1+=     0.00000002084 * cos( 5.84682067018 +     1577.34354244780*t);
+   emb_z_1+=     0.00000001959 * cos( 1.42270916062 +    12036.46073488820*t);
+   emb_z_1+=     0.00000001396 * cos( 0.50057569258 +     6309.37416979120*t);
+   emb_z_1+=     0.00000001584 * cos( 1.81185253086 +     5856.47765911540*t);
+   emb_z_1+=     0.00000001493 * cos( 1.75695182457 +     5884.92684658320*t);
+   emb_z_1+=     0.00000001713 * cos( 6.04661123493 +      398.14900340820*t);
+   emb_z_1+=     0.00000001298 * cos( 3.71598693704 +     1059.38193018920*t);
+   emb_z_1+=     0.00000001573 * cos( 2.20074999518 +     9437.76293488700*t);
+   emb_z_1+=     0.00000001503 * cos( 3.94648902715 +     4705.73230754360*t);
+   emb_z_1+=     0.00000001041 * cos( 2.94815263534 +     6256.77753019160*t);
+   emb_z_1+=     0.00000000983 * cos( 2.24865863102 +     6681.22485339960*t);
+   emb_z_1+=     0.00000001150 * cos( 2.43940396435 +      775.52261132400*t);
+   emb_z_1+=     0.00000001171 * cos( 0.74619691355 +     6812.76681508600*t);
+   emb_z_1+=     0.00000000881 * cos( 1.86103443124 +     3930.20969621960*t);
+   emb_z_1+=     0.00000000696 * cos( 1.45896954690 +    11790.62908865880*t);
+   emb_z_1+=     0.00000000732 * cos( 1.60348898963 +     4694.00295470760*t);
+   emb_z_1+=     0.00000000487 * cos( 0.04282116718 +    11506.76976979360*t);
+   emb_z_1+=     0.00000000632 * cos( 4.34207925908 +     5753.38488489680*t);
+   emb_z_1+=     0.00000000583 * cos( 3.92560802577 +    12168.00269657460*t);
+   emb_z_1+=     0.00000000409 * cos( 4.63047700185 +      796.29800681640*t);
+   emb_z_1+=     0.00000000483 * cos( 5.74980419879 +    17789.84561978500*t);
+   emb_z_1+=     0.00000000405 * cos( 1.97558286872 +     6283.00853968860*t);
+   emb_z_1+=     0.00000000405 * cos( 1.53147989887 +     6283.14316029419*t);
+   emb_z_1+=     0.00000000354 * cos( 3.28490469609 +    25132.30339996560*t);
+   emb_z_1+=     0.00000000417 * cos( 3.63631064623 +    13367.97263110660*t);
+   emb_z_1+=     0.00000000367 * cos( 1.55081387384 +    18073.70493865020*t);
+   emb_z_1+=     0.00000000284 * cos( 0.15061211477 +     3738.76143010800*t);
+   emb_z_1+=     0.00000000295 * cos( 2.75399172194 +    14143.49524243060*t);
+   emb_z_1+=     0.00000000288 * cos( 2.98400521096 +     6275.96230299060*t);
+   emb_z_1+=     0.00000000353 * cos( 2.42503121206 +      213.29909543800*t);
+   emb_z_1+=     0.00000000250 * cos( 0.42640065085 +     7079.37385680780*t);
+   emb_z_1+=     0.00000000251 * cos( 2.81315684448 +     6290.18939699220*t);
+   emb_z_1+=     0.00000000338 * cos( 2.55274440586 +     7058.59846131540*t);
+   emb_z_1+=     0.00000000230 * cos( 2.54182122331 +     7084.89678111520*t);
+   emb_z_1+=     0.00000000212 * cos( 4.71778046808 +     6496.37494542940*t);
+   emb_z_1+=     0.00000000188 * cos( 3.17403855561 +     5486.77784317500*t);
+   emb_z_1+=     0.00000000184 * cos( 4.29245958347 +     1589.07289528380*t);
+   emb_z_1+=     0.00000000216 * cos( 2.95054003179 +    22003.91463486980*t);
+   emb_z_1+=     0.00000000165 * cos( 1.19407165511 +     5643.17856367740*t);
+   emb_z_1+=     0.00000000162 * cos( 2.94100342992 +    15720.83878487840*t);
+   emb_z_1+=     0.00000000189 * cos( 5.09534739493 +    17298.18232732620*t);
+   emb_z_1+=     0.00000000149 * cos( 5.23057852418 +    11769.85369316640*t);
+   emb_z_1+=     0.00000000150 * cos( 1.65291444016 +     3128.38876509580*t);
+   emb_z_1+=     0.00000000143 * cos( 5.48362795500 +    10977.07880469900*t);
+   emb_z_1+=     0.00000000144 * cos( 3.81203756929 +     6262.30045449900*t);
+   emb_z_1+=     0.00000000143 * cos( 5.97809021355 +     6303.85124548380*t);
+   emb_z_1+=     0.00000000136 * cos( 4.97038154072 +     4136.91043351620*t);
+   emb_z_1+=     0.00000000163 * cos( 2.00627600582 +     5230.80746680300*t);
+   emb_z_1+=     0.00000000135 * cos( 3.68020950012 +        7.11354700080*t);
+   emb_z_1+=     0.00000000130 * cos( 4.41428061934 +     6127.65545055720*t);
+   emb_z_1+=     0.00000000123 * cos( 0.93547502051 +     3340.61242669980*t);
+   emb_z_1+=     0.00000000122 * cos( 3.13394683057 +     1194.44701022460*t);
+   emb_z_1+=     0.00000000131 * cos( 4.24961399629 +     6282.09552892320*t);
+   emb_z_1+=     0.00000000131 * cos( 5.54051222995 +     6284.05617105960*t);
+   emb_z_1+=     0.00000000161 * cos( 3.32421999691 +     6283.31966747490*t);
+   emb_z_1+=     0.00000000128 * cos( 4.35489873365 +    25934.12433108940*t);
+   emb_z_1+=     0.00000000113 * cos( 4.42578862670 +    19651.04848109800*t);
+   emb_z_1+=     0.00000000107 * cos( 2.97447343520 +     8827.39026987480*t);
+   emb_z_1+=     0.00000000107 * cos( 0.41198538606 +     5481.25491886760*t);
+   emb_z_1+=     0.00000000131 * cos( 2.29167404623 +     3154.68708489560*t);
+   emb_z_1+=     0.00000000139 * cos( 0.84923405785 +     6069.77675455340*t);
+   emb_z_1+=     0.00000000109 * cos( 3.21932289821 +      801.82093112380*t);
+   emb_z_1+=     0.00000000086 * cos( 2.78714560645 +     5216.58037280140*t);
+   emb_z_1+=     0.00000000116 * cos( 1.07208569199 +    12352.85260454480*t);
+   emb_z_1+=     0.00000000085 * cos( 4.74199868968 +     8429.24126646660*t);
+   emb_z_1+=     0.00000000112 * cos( 5.29899379351 +     2544.31441988340*t);
+   emb_z_1+=     0.00000000096 * cos( 5.20691825399 +    17260.15465469040*t);
+   emb_z_1+=     0.00000000074 * cos( 4.51170584596 +     8635.94200376320*t);
+   emb_z_1+=     0.00000000097 * cos( 5.50959692365 +    29864.33402730900*t);
+   emb_z_1+=     0.00000000094 * cos( 0.25007575444 +    21228.39202354580*t);
+   emb_z_1+=     0.00000000082 * cos( 3.01994555641 +    20426.57109242200*t);
+   emb_z_1+=     0.00000000071 * cos( 0.47420303374 +     2146.16541647520*t);
+   emb_z_1+=     0.00000000074 * cos( 1.27647360807 +     4164.31198961300*t);
+   emb_z_1+=     0.00000000075 * cos( 0.73893922793 +     9225.53927328300*t);
+   emb_z_1+=     0.00000000065 * cos( 1.39694123605 +    12569.67481833180*t);
+   emb_z_1+=     0.00000000074 * cos( 2.20184828895 +    37724.75341974820*t);
+   emb_z_1+=     0.00000000058 * cos( 5.36736393336 +    11015.10647733480*t);
+   emb_z_1+=     0.00000000069 * cos( 3.56483879750 +     4535.05943692440*t);
+   emb_z_1+=     0.00000000064 * cos( 2.05206957553 +    12562.62858163380*t);
+   emb_z_1+=     0.00000000050 * cos( 4.30582312752 +      426.59819087600*t);
+   emb_z_1+=     0.00000000056 * cos( 4.14424591372 +     5088.62883976680*t);
+   emb_z_1+=     0.00000000061 * cos( 0.63918772258 +    33794.54372352860*t);
+   emb_z_1+=     0.00000000049 * cos( 6.20603791837 +    12559.03815298200*t);
+   emb_z_1+=     0.00000000058 * cos( 1.68201284430 +    25158.60171976540*t);
+   emb_z_1+=     0.00000000051 * cos( 5.03600618727 +     5863.59120611620*t);
+   emb_z_1+=     0.00000000047 * cos( 0.30276119177 +    11371.70468975820*t);
+   emb_z_1+=     0.00000000050 * cos( 0.77635440262 +       26.29831979980*t);
+   emb_z_1+=     0.00000000041 * cos( 1.87344929139 +     1748.01641306700*t);
+   emb_z_1+=     0.00000000041 * cos( 5.91861144924 +    23581.25817731760*t);
+   emb_z_1+=     0.00000000055 * cos( 1.37242658028 +     7342.45778018060*t);
+   emb_z_1+=     0.00000000053 * cos( 1.47750772846 +    18319.53658487960*t);
+   emb_z_1+=     0.00000000042 * cos( 0.11693102263 +    23543.23050468179*t);
+   emb_z_1+=     0.00000000045 * cos( 5.18332998014 +    16496.36139620240*t);
+   emb_z_1+=     0.00000000045 * cos( 1.69577010458 +     6702.56049386660*t);
+   emb_z_1+=     0.00000000050 * cos( 3.62740835096 +    41654.96311596780*t);
+   emb_z_1+=     0.00000000050 * cos( 0.12043698146 +    29088.81141598500*t);
+   emb_z_1+=     0.00000000039 * cos( 2.15562950918 +     4933.20844033260*t);
+   emb_z_1+=     0.00000000048 * cos( 5.37118742490 +     7477.52286021600*t);
+   emb_z_1+=     0.00000000040 * cos( 3.02033323932 +    10447.38783960440*t);
+   emb_z_1+=     0.00000000035 * cos( 3.52010534689 +    26087.90314157420*t);
+   emb_z_1+=     0.00000000045 * cos( 6.10982446357 +     8031.09226305840*t);
+   emb_z_1=emb_z_1 * t;
+
+   double emb_z_2=0.0;
+   emb_z_2+=     0.00009721989 * cos( 5.15233725915 +     6283.07584999140*t);
+   emb_z_2+=     0.00000349501 * cos( 3.14159265359 +        0.00000000000*t);
+   emb_z_2+=     0.00000067136 * cos( 0.64403888586 +    12566.15169998280*t);
+   emb_z_2+=     0.00000002169 * cos( 1.07314429736 +    18849.22754997420*t);
+   emb_z_2+=     0.00000000635 * cos( 3.51985338656 +     6279.55273164240*t);
+   emb_z_2+=     0.00000000347 * cos( 6.16127984865 +     5223.69391980220*t);
+   emb_z_2+=     0.00000000322 * cos( 0.37660897389 +     6286.59896834040*t);
+   emb_z_2+=     0.00000000314 * cos( 2.37089871705 +     5507.55323866740*t);
+   emb_z_2+=     0.00000000242 * cos( 1.64810071161 +     7860.41939243920*t);
+   emb_z_2+=     0.00000000243 * cos( 3.96741961391 +      529.69096509460*t);
+   emb_z_2+=     0.00000000229 * cos( 0.50869157074 +     2352.86615377180*t);
+   emb_z_2+=     0.00000000170 * cos( 5.28668290523 +     6256.77753019160*t);
+   emb_z_2+=     0.00000000174 * cos( 0.13496987900 +    10213.28554621100*t);
+   emb_z_2+=     0.00000000125 * cos( 3.13799899281 +     5884.92684658320*t);
+   emb_z_2+=     0.00000000118 * cos( 3.85466055344 +     1577.34354244780*t);
+   emb_z_2+=     0.00000000123 * cos( 4.17102530625 +     6309.37416979120*t);
+   emb_z_2+=     0.00000000089 * cos( 1.65074757807 +     4705.73230754360*t);
+   emb_z_2+=     0.00000000090 * cos( 4.13333302298 +     5753.38488489680*t);
+   emb_z_2+=     0.00000000067 * cos( 0.30879908581 +      775.52261132400*t);
+   emb_z_2+=     0.00000000088 * cos( 2.93230657226 +    12036.46073488820*t);
+   emb_z_2+=     0.00000000065 * cos( 6.02619905507 +     6812.76681508600*t);
+   emb_z_2+=     0.00000000067 * cos( 1.15131045651 +    25132.30339996560*t);
+   emb_z_2+=     0.00000000065 * cos( 2.72055837855 +     6127.65545055720*t);
+   emb_z_2+=     0.00000000062 * cos( 0.78482367183 +     6438.49624942560*t);
+   emb_z_2+=     0.00000000057 * cos( 1.49017558261 +    11506.76976979360*t);
+   emb_z_2+=     0.00000000056 * cos( 4.55921214755 +     5486.77784317500*t);
+   emb_z_2+=     0.00000000069 * cos( 1.21340931171 +      398.14900340820*t);
+   emb_z_2+=     0.00000000047 * cos( 5.28931662006 +     7079.37385680780*t);
+   emb_z_2+=     0.00000000044 * cos( 3.09697268535 +    11790.62908865880*t);
+   emb_z_2+=     0.00000000040 * cos( 1.74073861374 +     1059.38193018920*t);
+   emb_z_2+=     0.00000000040 * cos( 2.70069983169 +    10988.80815753500*t);
+   emb_z_2+=     0.00000000036 * cos( 0.32483893578 +     6681.22485339960*t);
+   emb_z_2+=     0.00000000033 * cos( 1.26266496002 +     6290.18939699220*t);
+   emb_z_2+=     0.00000000033 * cos( 2.25067065498 +     6275.96230299060*t);
+   emb_z_2+=     0.00000000032 * cos( 2.90595681986 +      796.29800681640*t);
+   emb_z_2+=     0.00000000029 * cos( 1.96179670465 +        7.11354700080*t);
+   emb_z_2+=     0.00000000023 * cos( 4.86623778952 +     3738.76143010800*t);
+   emb_z_2+=     0.00000000021 * cos( 6.22275008403 +     6069.77675455340*t);
+   emb_z_2+=     0.00000000024 * cos( 5.83690442827 +     6282.09552892320*t);
+   emb_z_2+=     0.00000000024 * cos( 3.95322179797 +     6284.05617105960*t);
+   emb_z_2+=     0.00000000022 * cos( 3.46618662246 +     4136.91043351620*t);
+   emb_z_2+=     0.00000000021 * cos( 3.58048145762 +     6496.37494542940*t);
+   emb_z_2+=     0.00000000025 * cos( 5.59503085104 +     7058.59846131540*t);
+   emb_z_2+=     0.00000000019 * cos( 1.62978232282 +     1194.44701022460*t);
+   emb_z_2+=     0.00000000019 * cos( 2.82679050999 +     1589.07289528380*t);
+   emb_z_2+=     0.00000000024 * cos( 5.40715542605 +    12168.00269657460*t);
+   emb_z_2=emb_z_2 * t * t;
+
+   double emb_z_3=0.0;
+   emb_z_3+=     0.00000276077 * cos( 0.59413258730 +     6283.07584999140*t);
+   emb_z_3+=     0.00000025551 * cos( 3.14159265359 +        0.00000000000*t);
+   emb_z_3+=     0.00000001810 * cos( 0.11612262117 +    12566.15169998280*t);
+   emb_z_3+=     0.00000000113 * cos( 5.66152847225 +    18849.22754997420*t);
+   emb_z_3+=     0.00000000056 * cos( 5.02765554835 +     6279.55273164240*t);
+   emb_z_3+=     0.00000000019 * cos( 5.99007646261 +     6256.77753019160*t);
+   emb_z_3+=     0.00000000018 * cos( 3.80004734567 +     6309.37416979120*t);
+   emb_z_3+=     0.00000000018 * cos( 1.21049250774 +     6127.65545055720*t);
+   emb_z_3+=     0.00000000018 * cos( 2.29734567137 +     6438.49624942560*t);
+   emb_z_3+=     0.00000000015 * cos( 4.72881467263 +     6286.59896834040*t);
+   emb_z_3=emb_z_3 * t * t * t;
+
+   double emb_z_4=0.0;
+   emb_z_4+=     0.00000005751 * cos( 2.27069090892 +     6283.07584999140*t);
+   emb_z_4+=     0.00000001305 * cos( 0.00000000000 +        0.00000000000*t);
+   emb_z_4+=     0.00000000060 * cos( 4.26712050846 +    12566.15169998280*t);
+   emb_z_4+=     0.00000000006 * cos( 4.13914011953 +    18849.22754997420*t);
+   emb_z_4=emb_z_4 * t * t * t * t;
+
+   double emb_z_5=0.0;
+   emb_z_5+=     0.00000000114 * cos( 4.31169156923 +     6283.07584999140*t);
+   emb_z_5+=     0.00000000035 * cos( 0.00000000000 +        0.00000000000*t);
+   emb_z_5=emb_z_5 * t * t * t * t * t;
+
+   return emb_z_0+emb_z_1+emb_z_2+emb_z_3+emb_z_4+emb_z_5;
+}
\ No newline at end of file
