Added the first version of the Celestial Vault DaySequence plugin

diff --git a/Engine/Plugins/Experimental/CelestialVault/Source/CelestialVault/Private/vsop87a_milli.cpp b/Engine/Plugins/Experimental/CelestialVault/Source/CelestialVault/Private/vsop87a_milli.cpp
new file mode 100644
--- /dev/null
+++ b/Engine/Plugins/Experimental/CelestialVault/Source/CelestialVault/Private/vsop87a_milli.cpp
@@ -0,0 +1,2024 @@
+// Copyright Epic Games, Inc. All Rights Reserved.
+
+
+//VSOP87-Multilang http://www.celestialprogramming.com/vsop87-multilang/index.html
+//Greg Miller ([REDACTED_EMAIL]) 2019.  Released as Public Domain
+
+#include "vsop87a_milli.h"
+#include <math.h>
+
+
+void vsop87a_milli::getEarth(double t,double temp[]){
+   temp[0]=vsop87a_milli::earth_x(t);
+   temp[1]=vsop87a_milli::earth_y(t);
+   temp[2]=vsop87a_milli::earth_z(t);
+}
+
+void vsop87a_milli::getEmb(double t,double temp[]){
+   temp[0]=vsop87a_milli::emb_x(t);
+   temp[1]=vsop87a_milli::emb_y(t);
+   temp[2]=vsop87a_milli::emb_z(t);
+}
+
+void vsop87a_milli::getJupiter(double t,double temp[]){
+   temp[0]=vsop87a_milli::jupiter_x(t);
+   temp[1]=vsop87a_milli::jupiter_y(t);
+   temp[2]=vsop87a_milli::jupiter_z(t);
+}
+
+void vsop87a_milli::getMars(double t,double temp[]){
+   temp[0]=vsop87a_milli::mars_x(t);
+   temp[1]=vsop87a_milli::mars_y(t);
+   temp[2]=vsop87a_milli::mars_z(t);
+}
+
+void vsop87a_milli::getMercury(double t,double temp[]){
+   temp[0]=vsop87a_milli::mercury_x(t);
+   temp[1]=vsop87a_milli::mercury_y(t);
+   temp[2]=vsop87a_milli::mercury_z(t);
+}
+
+void vsop87a_milli::getNeptune(double t,double temp[]){
+   temp[0]=vsop87a_milli::neptune_x(t);
+   temp[1]=vsop87a_milli::neptune_y(t);
+   temp[2]=vsop87a_milli::neptune_z(t);
+}
+
+void vsop87a_milli::getSaturn(double t,double temp[]){
+   temp[0]=vsop87a_milli::saturn_x(t);
+   temp[1]=vsop87a_milli::saturn_y(t);
+   temp[2]=vsop87a_milli::saturn_z(t);
+}
+
+void vsop87a_milli::getUranus(double t,double temp[]){
+   temp[0]=vsop87a_milli::uranus_x(t);
+   temp[1]=vsop87a_milli::uranus_y(t);
+   temp[2]=vsop87a_milli::uranus_z(t);
+}
+
+void vsop87a_milli::getVenus(double t,double temp[]){
+   temp[0]=vsop87a_milli::venus_x(t);
+   temp[1]=vsop87a_milli::venus_y(t);
+   temp[2]=vsop87a_milli::venus_z(t);
+}
+
+void vsop87a_milli::getMoon(double earth[], double emb[],double temp[]){
+   temp[0]=(emb[0]-earth[0])*(1 + 1 / 0.01230073677);
+   temp[1]=(emb[1]-earth[1])*(1 + 1 / 0.01230073677);
+   temp[2]=(emb[2]-earth[2])*(1 + 1 / 0.01230073677);
+   temp[0]=temp[0]+earth[0];
+   temp[1]=temp[1]+earth[1];
+   temp[2]=temp[2]+earth[2];
+}
+
+double vsop87a_milli::earth_x(double t){
+   double earth_x_0=0.0;
+   earth_x_0+=     0.99982928844 * cos( 1.75348568475 +     6283.07584999140*t);
+   earth_x_0+=     0.00835257300 * cos( 1.71034539450 +    12566.15169998280*t);
+   earth_x_0+=     0.00561144206 * cos( 0.00000000000 +        0.00000000000*t);
+   earth_x_0+=     0.00010466628 * cos( 1.66722645223 +    18849.22754997420*t);
+   earth_x_0+=     0.00003110838 * cos( 0.66875185215 +    83996.84731811189*t);
+   earth_x_0+=     0.00002552498 * cos( 0.58310207301 +      529.69096509460*t);
+   earth_x_0+=     0.00002137256 * cos( 1.09235189672 +     1577.34354244780*t);
+   earth_x_0+=     0.00001709103 * cos( 0.49540223397 +     6279.55273164240*t);
+   earth_x_0+=     0.00001707882 * cos( 6.15315547484 +     6286.59896834040*t);
+   earth_x_0+=     0.00001445242 * cos( 3.47272783760 +     2352.86615377180*t);
+   earth_x_0+=     0.00001091006 * cos( 3.68984782465 +     5223.69391980220*t);
+
+   double earth_x_1=0.0;
+   earth_x_1+=     0.00123403056 * cos( 0.00000000000 +        0.00000000000*t);
+   earth_x_1+=     0.00051500156 * cos( 6.00266267204 +    12566.15169998280*t);
+   earth_x_1+=     0.00001290726 * cos( 5.95943124583 +    18849.22754997420*t);
+   earth_x_1+=     0.00001068627 * cos( 2.01554176551 +     6283.07584999140*t);
+   earth_x_1=earth_x_1 * t;
+
+   double earth_x_2=0.0;
+   earth_x_2+=     0.00004143217 * cos( 3.14159265359 +        0.00000000000*t);
+   earth_x_2+=     0.00002175695 * cos( 4.39999849572 +    12566.15169998280*t);
+   earth_x_2=earth_x_2 * t * t;
+
+   return earth_x_0+earth_x_1+earth_x_2;
+}
+
+double vsop87a_milli::earth_y(double t){
+   double earth_y_0=0.0;
+   earth_y_0+=     0.99989211030 * cos( 0.18265890456 +     6283.07584999140*t);
+   earth_y_0+=     0.02442699036 * cos( 3.14159265359 +        0.00000000000*t);
+   earth_y_0+=     0.00835292314 * cos( 0.13952878991 +    12566.15169998280*t);
+   earth_y_0+=     0.00010466965 * cos( 0.09641690558 +    18849.22754997420*t);
+   earth_y_0+=     0.00003110838 * cos( 5.38114091484 +    83996.84731811189*t);
+   earth_y_0+=     0.00002570338 * cos( 5.30103973360 +      529.69096509460*t);
+   earth_y_0+=     0.00002147473 * cos( 2.66253538905 +     1577.34354244780*t);
+   earth_y_0+=     0.00001709219 * cos( 5.20780401071 +     6279.55273164240*t);
+   earth_y_0+=     0.00001707987 * cos( 4.58232858766 +     6286.59896834040*t);
+   earth_y_0+=     0.00001440265 * cos( 1.90068164664 +     2352.86615377180*t);
+   earth_y_0+=     0.00001135092 * cos( 5.27313415220 +     5223.69391980220*t);
+
+   double earth_y_1=0.0;
+   earth_y_1+=     0.00093046324 * cos( 0.00000000000 +        0.00000000000*t);
+   earth_y_1+=     0.00051506609 * cos( 4.43180499286 +    12566.15169998280*t);
+   earth_y_1+=     0.00001290800 * cos( 4.38860548540 +    18849.22754997420*t);
+   earth_y_1=earth_y_1 * t;
+
+   double earth_y_2=0.0;
+   earth_y_2+=     0.00005080208 * cos( 0.00000000000 +        0.00000000000*t);
+   earth_y_2+=     0.00002178016 * cos( 2.82957544235 +    12566.15169998280*t);
+   earth_y_2+=     0.00001020487 * cos( 4.63746718598 +     6283.07584999140*t);
+   earth_y_2=earth_y_2 * t * t;
+
+   return earth_y_0+earth_y_1+earth_y_2;
+}
+
+double vsop87a_milli::earth_z(double t){
+   double earth_z_1=0.0;
+   earth_z_1+=     0.00227822442 * cos( 3.41372504278 +     6283.07584999140*t);
+   earth_z_1+=     0.00005429282 * cos( 0.00000000000 +        0.00000000000*t);
+   earth_z_1+=     0.00001903183 * cos( 3.37061270964 +    12566.15169998280*t);
+   earth_z_1=earth_z_1 * t;
+
+   double earth_z_2=0.0;
+   earth_z_2+=     0.00009721989 * cos( 5.15233725915 +     6283.07584999140*t);
+   earth_z_2=earth_z_2 * t * t;
+
+   return earth_z_1+earth_z_2;
+}
+
+double vsop87a_milli::emb_x(double t){
+   double emb_x_0=0.0;
+   emb_x_0+=     0.99982927460 * cos( 1.75348568475 +     6283.07584999140*t);
+   emb_x_0+=     0.00835257300 * cos( 1.71034539450 +    12566.15169998280*t);
+   emb_x_0+=     0.00561144161 * cos( 0.00000000000 +        0.00000000000*t);
+   emb_x_0+=     0.00010466628 * cos( 1.66722645223 +    18849.22754997420*t);
+   emb_x_0+=     0.00002552498 * cos( 0.58310207301 +      529.69096509460*t);
+   emb_x_0+=     0.00002137256 * cos( 1.09235189672 +     1577.34354244780*t);
+   emb_x_0+=     0.00001709103 * cos( 0.49540223397 +     6279.55273164240*t);
+   emb_x_0+=     0.00001707882 * cos( 6.15315547484 +     6286.59896834040*t);
+   emb_x_0+=     0.00001445242 * cos( 3.47272783760 +     2352.86615377180*t);
+   emb_x_0+=     0.00001091006 * cos( 3.68984782465 +     5223.69391980220*t);
+
+   double emb_x_1=0.0;
+   emb_x_1+=     0.00123403046 * cos( 0.00000000000 +        0.00000000000*t);
+   emb_x_1+=     0.00051500156 * cos( 6.00266267204 +    12566.15169998280*t);
+   emb_x_1+=     0.00001290726 * cos( 5.95943124583 +    18849.22754997420*t);
+   emb_x_1+=     0.00001068627 * cos( 2.01554176551 +     6283.07584999140*t);
+   emb_x_1=emb_x_1 * t;
+
+   double emb_x_2=0.0;
+   emb_x_2+=     0.00004143217 * cos( 3.14159265359 +        0.00000000000*t);
+   emb_x_2+=     0.00002175695 * cos( 4.39999849572 +    12566.15169998280*t);
+   emb_x_2=emb_x_2 * t * t;
+
+   return emb_x_0+emb_x_1+emb_x_2;
+}
+
+double vsop87a_milli::emb_y(double t){
+   double emb_y_0=0.0;
+   emb_y_0+=     0.99989209645 * cos( 0.18265890456 +     6283.07584999140*t);
+   emb_y_0+=     0.02442698841 * cos( 3.14159265359 +        0.00000000000*t);
+   emb_y_0+=     0.00835292314 * cos( 0.13952878991 +    12566.15169998280*t);
+   emb_y_0+=     0.00010466965 * cos( 0.09641690558 +    18849.22754997420*t);
+   emb_y_0+=     0.00002570338 * cos( 5.30103973360 +      529.69096509460*t);
+   emb_y_0+=     0.00002147473 * cos( 2.66253538905 +     1577.34354244780*t);
+   emb_y_0+=     0.00001709219 * cos( 5.20780401071 +     6279.55273164240*t);
+   emb_y_0+=     0.00001707987 * cos( 4.58232858766 +     6286.59896834040*t);
+   emb_y_0+=     0.00001440265 * cos( 1.90068164664 +     2352.86615377180*t);
+   emb_y_0+=     0.00001135092 * cos( 5.27313415220 +     5223.69391980220*t);
+
+   double emb_y_1=0.0;
+   emb_y_1+=     0.00093046317 * cos( 0.00000000000 +        0.00000000000*t);
+   emb_y_1+=     0.00051506609 * cos( 4.43180499286 +    12566.15169998280*t);
+   emb_y_1+=     0.00001290800 * cos( 4.38860548540 +    18849.22754997420*t);
+   emb_y_1=emb_y_1 * t;
+
+   double emb_y_2=0.0;
+   emb_y_2+=     0.00005080208 * cos( 0.00000000000 +        0.00000000000*t);
+   emb_y_2+=     0.00002178016 * cos( 2.82957544235 +    12566.15169998280*t);
+   emb_y_2+=     0.00001020487 * cos( 4.63746718598 +     6283.07584999140*t);
+   emb_y_2=emb_y_2 * t * t;
+
+   return emb_y_0+emb_y_1+emb_y_2;
+}
+
+double vsop87a_milli::emb_z(double t){
+   double emb_z_1=0.0;
+   emb_z_1+=     0.00227822442 * cos( 3.41372504278 +     6283.07584999140*t);
+   emb_z_1+=     0.00005429282 * cos( 0.00000000000 +        0.00000000000*t);
+   emb_z_1+=     0.00001903183 * cos( 3.37061270964 +    12566.15169998280*t);
+   emb_z_1=emb_z_1 * t;
+
+   double emb_z_2=0.0;
+   emb_z_2+=     0.00009721989 * cos( 5.15233725915 +     6283.07584999140*t);
+   emb_z_2=emb_z_2 * t * t;
+
+   return emb_z_1+emb_z_2;
+}
+
+double vsop87a_milli::jupiter_x(double t){
+   double jupiter_x_0=0.0;
+   jupiter_x_0+=     5.19663470114 * cos( 0.59945082355 +      529.69096509460*t);
+   jupiter_x_0+=     0.36662642320 * cos( 3.14159265359 +        0.00000000000*t);
+   jupiter_x_0+=     0.12593937922 * cos( 0.94911583701 +     1059.38193018920*t);
+   jupiter_x_0+=     0.01500672056 * cos( 0.73175134610 +      522.57741809380*t);
+   jupiter_x_0+=     0.01476224578 * cos( 3.61736921122 +      536.80451209540*t);
+   jupiter_x_0+=     0.00457752736 * cos( 1.29883700755 +     1589.07289528380*t);
+   jupiter_x_0+=     0.00301689798 * cos( 5.17372551148 +        7.11354700080*t);
+   jupiter_x_0+=     0.00385975375 * cos( 2.01229910687 +      103.09277421860*t);
+   jupiter_x_0+=     0.00194025405 * cos( 5.02580363996 +      426.59819087600*t);
+   jupiter_x_0+=     0.00150678793 * cos( 6.12003027739 +      110.20632121940*t);
+   jupiter_x_0+=     0.00144867641 * cos( 5.55980577080 +      632.78373931320*t);
+   jupiter_x_0+=     0.00134226996 * cos( 0.87648567011 +      213.29909543800*t);
+   jupiter_x_0+=     0.00103494641 * cos( 6.19324769120 +     1052.26838318840*t);
+   jupiter_x_0+=     0.00114201562 * cos( 0.01567084269 +     1162.47470440780*t);
+   jupiter_x_0+=     0.00072095575 * cos( 3.96117430643 +     1066.49547719000*t);
+   jupiter_x_0+=     0.00059486083 * cos( 4.45769374358 +      949.17560896980*t);
+   jupiter_x_0+=     0.00068284021 * cos( 3.44051122631 +      846.08283475120*t);
+   jupiter_x_0+=     0.00047092251 * cos( 1.44612636451 +      419.48464387520*t);
+   jupiter_x_0+=     0.00030623417 * cos( 2.99132321427 +      206.18554843720*t);
+   jupiter_x_0+=     0.00026613459 * cos( 4.85169906494 +      323.50541665740*t);
+   jupiter_x_0+=     0.00019727457 * cos( 1.64891626213 +     2118.76386037840*t);
+   jupiter_x_0+=     0.00016481594 * cos( 1.95150056568 +      316.39186965660*t);
+   jupiter_x_0+=     0.00016101974 * cos( 0.87973155980 +      515.46387109300*t);
+   jupiter_x_0+=     0.00014209487 * cos( 2.07769621413 +      742.99006053260*t);
+   jupiter_x_0+=     0.00015192516 * cos( 6.25820127906 +      735.87651353180*t);
+   jupiter_x_0+=     0.00011423199 * cos( 3.48146108929 +      543.91805909620*t);
+   jupiter_x_0+=     0.00012155285 * cos( 3.75229924999 +      525.75881183150*t);
+   jupiter_x_0+=     0.00011996271 * cos( 0.58568573729 +      533.62311835770*t);
+   jupiter_x_0+=     0.00008468556 * cos( 3.47248751739 +      639.89728631400*t);
+   jupiter_x_0+=     0.00008223302 * cos( 5.56680447143 +     1478.86657406440*t);
+   jupiter_x_0+=     0.00008694124 * cos( 0.38262009411 +     1692.16566950240*t);
+   jupiter_x_0+=     0.00007427517 * cos( 5.98380751196 +      956.28915597060*t);
+   jupiter_x_0+=     0.00007516470 * cos( 0.92896448412 +     1265.56747862640*t);
+   jupiter_x_0+=     0.00007655867 * cos( 0.14178789086 +     1581.95934828300*t);
+   jupiter_x_0+=     0.00005318791 * cos( 1.10494016349 +      526.50957135690*t);
+   jupiter_x_0+=     0.00005218492 * cos( 3.23235129224 +      532.87235883230*t);
+   jupiter_x_0+=     0.00005777311 * cos( 5.03726165628 +       14.22709400160*t);
+   jupiter_x_0+=     0.00004622685 * cos( 3.75817086099 +     1375.77379984580*t);
+   jupiter_x_0+=     0.00003939864 * cos( 4.30892687511 +     1596.18644228460*t);
+   jupiter_x_0+=     0.00004569444 * cos( 2.15087281710 +       95.97922721780*t);
+   jupiter_x_0+=     0.00002952712 * cos( 3.85988483947 +      309.27832265580*t);
+   jupiter_x_0+=     0.00002857935 * cos( 6.01118473739 +      117.31986822020*t);
+   jupiter_x_0+=     0.00002440094 * cos( 4.23995765702 +      433.71173787680*t);
+   jupiter_x_0+=     0.00002438257 * cos( 3.88808463822 +      220.41264243880*t);
+   jupiter_x_0+=     0.00002675112 * cos( 3.18723449094 +     1169.58825140860*t);
+   jupiter_x_0+=     0.00002386425 * cos( 5.96354994324 +     1045.15483618760*t);
+   jupiter_x_0+=     0.00001870097 * cos( 0.52019313301 +     1155.36115740700*t);
+   jupiter_x_0+=     0.00001939060 * cos( 5.91883412864 +      625.67019231240*t);
+   jupiter_x_0+=     0.00001631500 * cos( 4.41910383466 +      942.06206196900*t);
+   jupiter_x_0+=     0.00001451667 * cos( 5.76112706040 +      853.19638175200*t);
+   jupiter_x_0+=     0.00001361286 * cos( 1.34792748837 +     1368.66025284500*t);
+   jupiter_x_0+=     0.00001663331 * cos( 1.94010629194 +      838.96928775040*t);
+   jupiter_x_0+=     0.00001611229 * cos( 5.49324974845 +       74.78159856730*t);
+   jupiter_x_0+=     0.00001033570 * cos( 0.08907208789 +     1795.25844372100*t);
+   jupiter_x_0+=     0.00001088284 * cos( 1.13406104190 +      527.24328453980*t);
+   jupiter_x_0+=     0.00001080643 * cos( 3.20528362573 +      532.13864564940*t);
+
+   double jupiter_x_1=0.0;
+   jupiter_x_1+=     0.00882389251 * cos( 3.14159265359 +        0.00000000000*t);
+   jupiter_x_1+=     0.00635297172 * cos( 0.10662156868 +     1059.38193018920*t);
+   jupiter_x_1+=     0.00599720482 * cos( 2.42996678275 +      522.57741809380*t);
+   jupiter_x_1+=     0.00589157060 * cos( 1.91556314637 +      536.80451209540*t);
+   jupiter_x_1+=     0.00081697204 * cos( 3.46668108797 +        7.11354700080*t);
+   jupiter_x_1+=     0.00046201898 * cos( 0.45714214032 +     1589.07289528380*t);
+   jupiter_x_1+=     0.00032508590 * cos( 1.74648849928 +     1052.26838318840*t);
+   jupiter_x_1+=     0.00033891193 * cos( 4.10113482752 +      529.69096509460*t);
+   jupiter_x_1+=     0.00031234303 * cos( 2.34698051502 +     1066.49547719000*t);
+   jupiter_x_1+=     0.00021244363 * cos( 4.36576178953 +      110.20632121940*t);
+   jupiter_x_1+=     0.00018156701 * cos( 4.00572238779 +      426.59819087600*t);
+   jupiter_x_1+=     0.00013577576 * cos( 0.30008010246 +      632.78373931320*t);
+   jupiter_x_1+=     0.00012889505 * cos( 2.57489294062 +      515.46387109300*t);
+   jupiter_x_1+=     0.00009125875 * cos( 1.78082469962 +      543.91805909620*t);
+   jupiter_x_1+=     0.00008085991 * cos( 6.16136518902 +      949.17560896980*t);
+   jupiter_x_1+=     0.00007142547 * cos( 3.17267801203 +      323.50541665740*t);
+   jupiter_x_1+=     0.00004292240 * cos( 4.74970626655 +      206.18554843720*t);
+   jupiter_x_1+=     0.00004393977 * cos( 1.14770788063 +      735.87651353180*t);
+   jupiter_x_1+=     0.00003399164 * cos( 2.90091450747 +      526.50957135690*t);
+   jupiter_x_1+=     0.00003333344 * cos( 1.43691652967 +      532.87235883230*t);
+   jupiter_x_1+=     0.00003873467 * cos( 3.33648870101 +       14.22709400160*t);
+   jupiter_x_1+=     0.00003044408 * cos( 1.65428048669 +      525.75881183150*t);
+   jupiter_x_1+=     0.00003001874 * cos( 2.68376982746 +      533.62311835770*t);
+   jupiter_x_1+=     0.00002933359 * cos( 2.61899855005 +      419.48464387520*t);
+   jupiter_x_1+=     0.00002438199 * cos( 3.60655644537 +      316.39186965660*t);
+   jupiter_x_1+=     0.00002804218 * cos( 4.89742591320 +      103.09277421860*t);
+   jupiter_x_1+=     0.00002990245 * cos( 0.80692155639 +     2118.76386037840*t);
+   jupiter_x_1+=     0.00001977572 * cos( 5.08915489088 +      956.28915597060*t);
+   jupiter_x_1+=     0.00001853679 * cos( 2.76941001747 +     1596.18644228460*t);
+   jupiter_x_1+=     0.00001772800 * cos( 0.72631739446 +      742.99006053260*t);
+   jupiter_x_1+=     0.00001812965 * cos( 3.84602148747 +       95.97922721780*t);
+   jupiter_x_1+=     0.00001532945 * cos( 4.31556714501 +      117.31986822020*t);
+   jupiter_x_1+=     0.00001904067 * cos( 1.85937873703 +     1581.95934828300*t);
+   jupiter_x_1+=     0.00001539212 * cos( 1.47899172821 +      639.89728631400*t);
+   jupiter_x_1+=     0.00001632362 * cos( 1.41504212408 +     1045.15483618760*t);
+   jupiter_x_1+=     0.00001023812 * cos( 2.57182697715 +      433.71173787680*t);
+   jupiter_x_1+=     0.00001055422 * cos( 2.50844222977 +     1265.56747862640*t);
+   jupiter_x_1=jupiter_x_1 * t;
+
+   double jupiter_x_2=0.0;
+   jupiter_x_2+=     0.00123864644 * cos( 4.13563277513 +      522.57741809380*t);
+   jupiter_x_2+=     0.00121521296 * cos( 0.21155109275 +      536.80451209540*t);
+   jupiter_x_2+=     0.00085355503 * cos( 0.00000000000 +        0.00000000000*t);
+   jupiter_x_2+=     0.00077685547 * cos( 5.29776154458 +      529.69096509460*t);
+   jupiter_x_2+=     0.00041410887 * cos( 5.12291589939 +     1059.38193018920*t);
+   jupiter_x_2+=     0.00011423070 * cos( 1.72917878238 +        7.11354700080*t);
+   jupiter_x_2+=     0.00007051587 * cos( 0.74163703419 +     1066.49547719000*t);
+   jupiter_x_2+=     0.00005711029 * cos( 3.63172846494 +     1052.26838318840*t);
+   jupiter_x_2+=     0.00005242644 * cos( 4.27482379441 +      515.46387109300*t);
+   jupiter_x_2+=     0.00004039540 * cos( 5.58417732117 +     1589.07289528380*t);
+   jupiter_x_2+=     0.00003706457 * cos( 0.07769981349 +      543.91805909620*t);
+   jupiter_x_2+=     0.00001698817 * cos( 2.44284418066 +      110.20632121940*t);
+   jupiter_x_2+=     0.00001134598 * cos( 2.35807061809 +      426.59819087600*t);
+   jupiter_x_2+=     0.00001322673 * cos( 1.63142549980 +       14.22709400160*t);
+   jupiter_x_2=jupiter_x_2 * t * t;
+
+   double jupiter_x_3=0.0;
+   jupiter_x_3+=     0.00017071323 * cos( 5.86133022278 +      522.57741809380*t);
+   jupiter_x_3+=     0.00016713548 * cos( 4.77458794485 +      536.80451209540*t);
+   jupiter_x_3+=     0.00003348610 * cos( 0.00000000000 +        0.00000000000*t);
+   jupiter_x_3+=     0.00001787838 * cos( 3.56550298031 +     1059.38193018920*t);
+   jupiter_x_3+=     0.00001435449 * cos( 5.98502036587 +      515.46387109300*t);
+   jupiter_x_3+=     0.00001080194 * cos( 5.42530305914 +     1066.49547719000*t);
+   jupiter_x_3+=     0.00001014206 * cos( 4.64773902077 +      543.91805909620*t);
+   jupiter_x_3+=     0.00001073175 * cos( 6.22314467964 +        7.11354700080*t);
+   jupiter_x_3=jupiter_x_3 * t * t * t;
+
+   double jupiter_x_4=0.0;
+   jupiter_x_4+=     0.00001762402 * cos( 1.32863039757 +      522.57741809380*t);
+   jupiter_x_4+=     0.00001717846 * cos( 3.03331531843 +      536.80451209540*t);
+   jupiter_x_4=jupiter_x_4 * t * t * t * t;
+
+   return jupiter_x_0+jupiter_x_1+jupiter_x_2+jupiter_x_3+jupiter_x_4;
+}
+
+double vsop87a_milli::jupiter_y(double t){
+   double jupiter_y_0=0.0;
+   jupiter_y_0+=     5.19520046589 * cos( 5.31203162731 +      529.69096509460*t);
+   jupiter_y_0+=     0.12592862602 * cos( 5.66160227728 +     1059.38193018920*t);
+   jupiter_y_0+=     0.09363670616 * cos( 3.14159265359 +        0.00000000000*t);
+   jupiter_y_0+=     0.01508275299 * cos( 5.43934968102 +      522.57741809380*t);
+   jupiter_y_0+=     0.01475809370 * cos( 2.04679566495 +      536.80451209540*t);
+   jupiter_y_0+=     0.00457750806 * cos( 6.01129093501 +     1589.07289528380*t);
+   jupiter_y_0+=     0.00300686679 * cos( 3.60948050740 +        7.11354700080*t);
+   jupiter_y_0+=     0.00378285578 * cos( 3.53006782383 +      103.09277421860*t);
+   jupiter_y_0+=     0.00192333128 * cos( 3.45690564771 +      426.59819087600*t);
+   jupiter_y_0+=     0.00146104656 * cos( 4.62267224431 +      110.20632121940*t);
+   jupiter_y_0+=     0.00139480058 * cos( 4.00075307706 +      632.78373931320*t);
+   jupiter_y_0+=     0.00132696764 * cos( 5.62184581859 +      213.29909543800*t);
+   jupiter_y_0+=     0.00101999807 * cos( 4.57594598884 +     1052.26838318840*t);
+   jupiter_y_0+=     0.00114043110 * cos( 4.72982262969 +     1162.47470440780*t);
+   jupiter_y_0+=     0.00072091178 * cos( 2.39048659148 +     1066.49547719000*t);
+   jupiter_y_0+=     0.00059051769 * cos( 2.89529070968 +      949.17560896980*t);
+   jupiter_y_0+=     0.00068374489 * cos( 1.86537074374 +      846.08283475120*t);
+   jupiter_y_0+=     0.00029807369 * cos( 4.52105772740 +      206.18554843720*t);
+   jupiter_y_0+=     0.00026933579 * cos( 3.86233956827 +      419.48464387520*t);
+   jupiter_y_0+=     0.00026619714 * cos( 3.28203174951 +      323.50541665740*t);
+   jupiter_y_0+=     0.00020873780 * cos( 3.79369881757 +      735.87651353180*t);
+   jupiter_y_0+=     0.00019727397 * cos( 0.07818534532 +     2118.76386037840*t);
+   jupiter_y_0+=     0.00018639846 * cos( 0.38751972138 +      316.39186965660*t);
+   jupiter_y_0+=     0.00016355726 * cos( 5.56997881604 +      515.46387109300*t);
+   jupiter_y_0+=     0.00014606858 * cos( 0.47759399145 +      742.99006053260*t);
+   jupiter_y_0+=     0.00011419853 * cos( 1.91089341468 +      543.91805909620*t);
+   jupiter_y_0+=     0.00012153427 * cos( 2.18151972499 +      525.75881183150*t);
+   jupiter_y_0+=     0.00011988875 * cos( 5.29687602089 +      533.62311835770*t);
+   jupiter_y_0+=     0.00008443107 * cos( 1.91435801697 +      639.89728631400*t);
+   jupiter_y_0+=     0.00008163163 * cos( 4.00303742375 +     1478.86657406440*t);
+   jupiter_y_0+=     0.00008732789 * cos( 5.09607066097 +     1692.16566950240*t);
+   jupiter_y_0+=     0.00007414115 * cos( 4.41141990461 +      956.28915597060*t);
+   jupiter_y_0+=     0.00007619486 * cos( 5.59554151997 +     1265.56747862640*t);
+   jupiter_y_0+=     0.00007779184 * cos( 4.83346300662 +     1581.95934828300*t);
+   jupiter_y_0+=     0.00005322882 * cos( 5.81740472645 +      526.50957135690*t);
+   jupiter_y_0+=     0.00005217025 * cos( 1.66178643542 +      532.87235883230*t);
+   jupiter_y_0+=     0.00005772132 * cos( 3.46915716927 +       14.22709400160*t);
+   jupiter_y_0+=     0.00004528355 * cos( 2.18377558038 +     1375.77379984580*t);
+   jupiter_y_0+=     0.00003939875 * cos( 2.73830531054 +     1596.18644228460*t);
+   jupiter_y_0+=     0.00004567181 * cos( 3.71300776935 +       95.97922721780*t);
+   jupiter_y_0+=     0.00003235419 * cos( 4.76600347062 +      625.67019231240*t);
+   jupiter_y_0+=     0.00003140740 * cos( 5.59566796922 +      309.27832265580*t);
+   jupiter_y_0+=     0.00002855423 * cos( 4.44478286006 +      117.31986822020*t);
+   jupiter_y_0+=     0.00002445625 * cos( 2.67036952230 +      433.71173787680*t);
+   jupiter_y_0+=     0.00002253545 * cos( 4.28462825722 +      838.96928775040*t);
+   jupiter_y_0+=     0.00002672262 * cos( 1.61857897069 +     1169.58825140860*t);
+   jupiter_y_0+=     0.00002423639 * cos( 2.32942339839 +      220.41264243880*t);
+   jupiter_y_0+=     0.00002362662 * cos( 4.60417580207 +     1155.36115740700*t);
+   jupiter_y_0+=     0.00002409581 * cos( 4.33196301609 +     1045.15483618760*t);
+   jupiter_y_0+=     0.00001458169 * cos( 4.18761881277 +      853.19638175200*t);
+   jupiter_y_0+=     0.00001432195 * cos( 3.24824554500 +      942.06206196900*t);
+   jupiter_y_0+=     0.00001646568 * cos( 3.91965876562 +       74.78159856730*t);
+   jupiter_y_0+=     0.00001050270 * cos( 4.83706014327 +     1795.25844372100*t);
+   jupiter_y_0+=     0.00001002355 * cos( 1.50931939870 +     1272.68102562720*t);
+   jupiter_y_0+=     0.00001087727 * cos( 5.84673086939 +      527.24328453980*t);
+   jupiter_y_0+=     0.00001079512 * cos( 1.63448507346 +      532.13864564940*t);
+
+   double jupiter_y_1=0.0;
+   jupiter_y_1+=     0.01694798253 * cos( 3.14159265359 +        0.00000000000*t);
+   jupiter_y_1+=     0.00634859798 * cos( 4.81903199650 +     1059.38193018920*t);
+   jupiter_y_1+=     0.00601160431 * cos( 0.85811249940 +      522.57741809380*t);
+   jupiter_y_1+=     0.00588928504 * cos( 0.34491576890 +      536.80451209540*t);
+   jupiter_y_1+=     0.00081187145 * cos( 1.90914316532 +        7.11354700080*t);
+   jupiter_y_1+=     0.00046888090 * cos( 1.91294535618 +      529.69096509460*t);
+   jupiter_y_1+=     0.00046194129 * cos( 5.16955994561 +     1589.07289528380*t);
+   jupiter_y_1+=     0.00032503453 * cos( 0.17640743623 +     1052.26838318840*t);
+   jupiter_y_1+=     0.00031231694 * cos( 0.77623645597 +     1066.49547719000*t);
+   jupiter_y_1+=     0.00019462096 * cos( 3.00957119470 +      110.20632121940*t);
+   jupiter_y_1+=     0.00017738615 * cos( 2.46531787101 +      426.59819087600*t);
+   jupiter_y_1+=     0.00013701692 * cos( 5.02070197804 +      632.78373931320*t);
+   jupiter_y_1+=     0.00013034616 * cos( 0.98979834442 +      515.46387109300*t);
+   jupiter_y_1+=     0.00009122660 * cos( 0.21022587969 +      543.91805909620*t);
+   jupiter_y_1+=     0.00008109050 * cos( 4.58123811601 +      949.17560896980*t);
+   jupiter_y_1+=     0.00007145229 * cos( 1.60381236094 +      323.50541665740*t);
+   jupiter_y_1+=     0.00003957592 * cos( 6.18550697817 +      206.18554843720*t);
+   jupiter_y_1+=     0.00004347346 * cos( 5.85522835488 +      735.87651353180*t);
+   jupiter_y_1+=     0.00003401735 * cos( 1.33033225252 +      526.50957135690*t);
+   jupiter_y_1+=     0.00003331887 * cos( 6.14951835712 +      532.87235883230*t);
+   jupiter_y_1+=     0.00003866147 * cos( 1.76877582038 +       14.22709400160*t);
+   jupiter_y_1+=     0.00003094257 * cos( 1.00670454701 +      419.48464387520*t);
+   jupiter_y_1+=     0.00003044205 * cos( 0.08329779827 +      525.75881183150*t);
+   jupiter_y_1+=     0.00003001484 * cos( 1.11280606283 +      533.62311835770*t);
+   jupiter_y_1+=     0.00002977284 * cos( 3.35507028507 +      103.09277421860*t);
+   jupiter_y_1+=     0.00002347100 * cos( 2.06781775390 +      316.39186965660*t);
+   jupiter_y_1+=     0.00002990192 * cos( 5.51944830506 +     2118.76386037840*t);
+   jupiter_y_1+=     0.00001875464 * cos( 5.32657356489 +      742.99006053260*t);
+   jupiter_y_1+=     0.00001854067 * cos( 1.19908734197 +     1596.18644228460*t);
+   jupiter_y_1+=     0.00001968401 * cos( 3.51896739844 +      956.28915597060*t);
+   jupiter_y_1+=     0.00001808627 * cos( 5.40287543026 +       95.97922721780*t);
+   jupiter_y_1+=     0.00001530472 * cos( 2.75094722237 +      117.31986822020*t);
+   jupiter_y_1+=     0.00001885393 * cos( 0.29905973710 +     1581.95934828300*t);
+   jupiter_y_1+=     0.00001516541 * cos( 6.21684203571 +      639.89728631400*t);
+   jupiter_y_1+=     0.00001636913 * cos( 6.09270756447 +     1045.15483618760*t);
+   jupiter_y_1+=     0.00001260123 * cos( 0.07143173954 +      625.67019231240*t);
+   jupiter_y_1+=     0.00001028165 * cos( 1.00301485824 +      433.71173787680*t);
+   jupiter_y_1+=     0.00001035933 * cos( 0.98273794152 +     1265.56747862640*t);
+   jupiter_y_1=jupiter_y_1 * t;
+
+   double jupiter_y_2=0.0;
+   jupiter_y_2+=     0.00124032509 * cos( 2.56495576833 +      522.57741809380*t);
+   jupiter_y_2+=     0.00121455991 * cos( 4.92398766380 +      536.80451209540*t);
+   jupiter_y_2+=     0.00076523263 * cos( 3.75913371793 +      529.69096509460*t);
+   jupiter_y_2+=     0.00076943042 * cos( 3.14159265359 +        0.00000000000*t);
+   jupiter_y_2+=     0.00041357600 * cos( 3.55228440457 +     1059.38193018920*t);
+   jupiter_y_2+=     0.00011277667 * cos( 0.18559902389 +        7.11354700080*t);
+   jupiter_y_2+=     0.00007051103 * cos( 5.45404368570 +     1066.49547719000*t);
+   jupiter_y_2+=     0.00005719440 * cos( 2.05970000230 +     1052.26838318840*t);
+   jupiter_y_2+=     0.00005286157 * cos( 2.69490465064 +      515.46387109300*t);
+   jupiter_y_2+=     0.00004039038 * cos( 4.01341034637 +     1589.07289528380*t);
+   jupiter_y_2+=     0.00003704528 * cos( 4.79029292271 +      543.91805909620*t);
+   jupiter_y_2+=     0.00001280283 * cos( 1.47574006861 +      110.20632121940*t);
+   jupiter_y_2+=     0.00001059783 * cos( 0.89610748176 +      426.59819087600*t);
+   jupiter_y_2+=     0.00001320627 * cos( 0.05786048417 +       14.22709400160*t);
+   jupiter_y_2=jupiter_y_2 * t * t;
+
+   double jupiter_y_3=0.0;
+   jupiter_y_3+=     0.00017085516 * cos( 4.29096904063 +      522.57741809380*t);
+   jupiter_y_3+=     0.00016701353 * cos( 3.20365737109 +      536.80451209540*t);
+   jupiter_y_3+=     0.00004006038 * cos( 0.00000000000 +        0.00000000000*t);
+   jupiter_y_3+=     0.00001782451 * cos( 1.99283071153 +     1059.38193018920*t);
+   jupiter_y_3+=     0.00001443816 * cos( 4.40866555269 +      515.46387109300*t);
+   jupiter_y_3+=     0.00001079405 * cos( 3.85450799252 +     1066.49547719000*t);
+   jupiter_y_3+=     0.00001013157 * cos( 3.07729621279 +      543.91805909620*t);
+   jupiter_y_3+=     0.00001055565 * cos( 4.70184773789 +        7.11354700080*t);
+   jupiter_y_3=jupiter_y_3 * t * t * t;
+
+   double jupiter_y_4=0.0;
+   jupiter_y_4+=     0.00001762645 * cos( 6.04159386554 +      522.57741809380*t);
+   jupiter_y_4+=     0.00001716045 * cos( 1.46206285710 +      536.80451209540*t);
+   jupiter_y_4=jupiter_y_4 * t * t * t * t;
+
+   return jupiter_y_0+jupiter_y_1+jupiter_y_2+jupiter_y_3+jupiter_y_4;
+}
+
+double vsop87a_milli::jupiter_z(double t){
+   double jupiter_z_0=0.0;
+   jupiter_z_0+=     0.11823100489 * cos( 3.55844646343 +      529.69096509460*t);
+   jupiter_z_0+=     0.00859031952 * cos( 0.00000000000 +        0.00000000000*t);
+   jupiter_z_0+=     0.00286562094 * cos( 3.90812238338 +     1059.38193018920*t);
+   jupiter_z_0+=     0.00042388592 * cos( 3.60144191032 +      522.57741809380*t);
+   jupiter_z_0+=     0.00033295491 * cos( 0.30297050585 +      536.80451209540*t);
+   jupiter_z_0+=     0.00010416160 * cos( 4.25764593061 +     1589.07289528380*t);
+   jupiter_z_0+=     0.00007449294 * cos( 5.24213104150 +      103.09277421860*t);
+   jupiter_z_0+=     0.00006910102 * cos( 1.75032945752 +        7.11354700080*t);
+   jupiter_z_0+=     0.00005292012 * cos( 1.68231447192 +      426.59819087600*t);
+   jupiter_z_0+=     0.00004313598 * cos( 3.70673689841 +      213.29909543800*t);
+   jupiter_z_0+=     0.00003784265 * cos( 2.71522544491 +      110.20632121940*t);
+   jupiter_z_0+=     0.00003798016 * cos( 2.16715743175 +      632.78373931320*t);
+   jupiter_z_0+=     0.00002455385 * cos( 2.96904135659 +     1052.26838318840*t);
+   jupiter_z_0+=     0.00002461547 * cos( 2.99889460411 +     1162.47470440780*t);
+   jupiter_z_0+=     0.00002001451 * cos( 2.68535838309 +      419.48464387520*t);
+   jupiter_z_0+=     0.00002163471 * cos( 6.26718259854 +      846.08283475120*t);
+   jupiter_z_0+=     0.00001633653 * cos( 0.64194743493 +     1066.49547719000*t);
+   jupiter_z_0+=     0.00001450672 * cos( 1.17108416193 +      949.17560896980*t);
+
+   double jupiter_z_1=0.0;
+   jupiter_z_1+=     0.00407072175 * cos( 1.52699353482 +      529.69096509460*t);
+   jupiter_z_1+=     0.00020307341 * cos( 2.59878269248 +     1059.38193018920*t);
+   jupiter_z_1+=     0.00014424953 * cos( 4.85400155025 +      536.80451209540*t);
+   jupiter_z_1+=     0.00015474611 * cos( 0.00000000000 +        0.00000000000*t);
+   jupiter_z_1+=     0.00012730364 * cos( 5.45536715732 +      522.57741809380*t);
+   jupiter_z_1+=     0.00002100882 * cos( 0.09538864287 +        7.11354700080*t);
+   jupiter_z_1+=     0.00001230425 * cos( 3.14222500244 +     1589.07289528380*t);
+   jupiter_z_1=jupiter_z_1 * t;
+
+   double jupiter_z_2=0.0;
+   jupiter_z_2+=     0.00028635326 * cos( 3.01374166973 +      529.69096509460*t);
+   jupiter_z_2+=     0.00003114752 * cos( 3.13228646176 +      536.80451209540*t);
+   jupiter_z_2+=     0.00002379765 * cos( 0.95574345340 +      522.57741809380*t);
+   jupiter_z_2+=     0.00001310111 * cos( 2.05263704913 +     1059.38193018920*t);
+   jupiter_z_2=jupiter_z_2 * t * t;
+
+   return jupiter_z_0+jupiter_z_1+jupiter_z_2;
+}
+
+double vsop87a_milli::mars_x(double t){
+   double mars_x_0=0.0;
+   mars_x_0+=     1.51769936383 * cos( 6.20403346548 +     3340.61242669980*t);
+   mars_x_0+=     0.19502945246 * cos( 3.14159265359 +        0.00000000000*t);
+   mars_x_0+=     0.07070919655 * cos( 0.25870338558 +     6681.22485339960*t);
+   mars_x_0+=     0.00494196914 * cos( 0.59669127768 +    10021.83728009940*t);
+   mars_x_0+=     0.00040938237 * cos( 0.93473307419 +    13362.44970679920*t);
+   mars_x_0+=     0.00021067199 * cos( 1.80435656154 +     3337.08930835080*t);
+   mars_x_0+=     0.00021041626 * cos( 1.17895619474 +     3344.13554504880*t);
+   mars_x_0+=     0.00011370375 * cos( 4.83265211109 +     1059.38193018920*t);
+   mars_x_0+=     0.00013527976 * cos( 0.63010765169 +      529.69096509460*t);
+   mars_x_0+=     0.00006774107 * cos( 3.61785048282 +     3340.59517304760*t);
+   mars_x_0+=     0.00006774060 * cos( 5.64862211431 +     3340.62968035200*t);
+   mars_x_0+=     0.00008226069 * cos( 1.86843519535 +     2281.23049651060*t);
+   mars_x_0+=     0.00005469046 * cos( 1.13324429003 +     2942.46342329160*t);
+   mars_x_0+=     0.00004817134 * cos( 1.85091045536 +     3738.76143010800*t);
+   mars_x_0+=     0.00004937579 * cos( 4.43241440654 +     5621.84292321040*t);
+   mars_x_0+=     0.00005276260 * cos( 2.33148083116 +     6151.53388830500*t);
+   mars_x_0+=     0.00003636667 * cos( 6.11397592106 +      796.29800681640*t);
+   mars_x_0+=     0.00003725823 * cos( 1.27280182943 +    16703.06213349900*t);
+   mars_x_0+=     0.00003729746 * cos( 1.21398323637 +      398.14900340820*t);
+   mars_x_0+=     0.00002368513 * cos( 2.96841895360 +     2544.31441988340*t);
+   mars_x_0+=     0.00002397865 * cos( 0.63553674054 +     3149.16416058820*t);
+   mars_x_0+=     0.00002274646 * cos( 2.35708328853 +     3532.06069281140*t);
+   mars_x_0+=     0.00001977579 * cos( 2.14087826110 +     6677.70173505060*t);
+   mars_x_0+=     0.00002229176 * cos( 1.69588962513 +     3340.54511639700*t);
+   mars_x_0+=     0.00002229117 * cos( 1.28739323821 +     3340.67973700260*t);
+   mars_x_0+=     0.00002182206 * cos( 1.69655112969 +     6283.07584999140*t);
+   mars_x_0+=     0.00002241010 * cos( 4.82218655311 +     8962.45534991020*t);
+   mars_x_0+=     0.00001677693 * cos( 3.14442612046 +     5884.92684658320*t);
+   mars_x_0+=     0.00001630482 * cos( 0.24117974845 +     4136.91043351620*t);
+   mars_x_0+=     0.00001958162 * cos( 1.51914544555 +     6684.74797174860*t);
+   mars_x_0+=     0.00001378470 * cos( 2.18011900021 +     1751.53953141600*t);
+   mars_x_0+=     0.00001289804 * cos( 4.70970778621 +     1194.44701022460*t);
+   mars_x_0+=     0.00001468124 * cos( 1.87869730543 +     3870.30339179440*t);
+   mars_x_0+=     0.00001290170 * cos( 0.43596325296 +     2810.92146160520*t);
+   mars_x_0+=     0.00001572540 * cos( 4.84809921789 +     1589.07289528380*t);
+
+   double mars_x_1=0.0;
+   mars_x_1+=     0.00861441374 * cos( 3.14159265359 +        0.00000000000*t);
+   mars_x_1+=     0.00552437949 * cos( 5.09565872891 +     6681.22485339960*t);
+   mars_x_1+=     0.00077184977 * cos( 5.43315636209 +    10021.83728009940*t);
+   mars_x_1+=     0.00020467294 * cos( 5.57051812369 +     3340.61242669980*t);
+   mars_x_1+=     0.00009589581 * cos( 5.77107234791 +    13362.44970679920*t);
+   mars_x_1+=     0.00002620610 * cos( 6.22441295122 +     3344.13554504880*t);
+   mars_x_1+=     0.00002620537 * cos( 3.04172154436 +     3337.08930835080*t);
+   mars_x_1+=     0.00001163612 * cos( 6.10909257097 +    16703.06213349900*t);
+   mars_x_1=mars_x_1 * t;
+
+   double mars_x_2=0.0;
+   mars_x_2+=     0.00056323939 * cos( 0.00000000000 +        0.00000000000*t);
+   mars_x_2+=     0.00022122528 * cos( 3.54372113272 +     6681.22485339960*t);
+   mars_x_2+=     0.00006091409 * cos( 3.93272649649 +    10021.83728009940*t);
+   mars_x_2+=     0.00001451998 * cos( 3.64655666460 +     3340.61242669980*t);
+   mars_x_2+=     0.00001130613 * cos( 4.28827023222 +    13362.44970679920*t);
+   mars_x_2=mars_x_2 * t * t;
+
+   return mars_x_0+mars_x_1+mars_x_2;
+}
+
+double vsop87a_milli::mars_y(double t){
+   double mars_y_0=0.0;
+   mars_y_0+=     1.51558976277 * cos( 4.63212206588 +     3340.61242669980*t);
+   mars_y_0+=     0.07064550239 * cos( 4.97051892902 +     6681.22485339960*t);
+   mars_y_0+=     0.08655481102 * cos( 0.00000000000 +        0.00000000000*t);
+   mars_y_0+=     0.00493872848 * cos( 5.30877806694 +    10021.83728009940*t);
+   mars_y_0+=     0.00040917422 * cos( 5.64698263703 +    13362.44970679920*t);
+   mars_y_0+=     0.00021036784 * cos( 0.23240270955 +     3337.08930835080*t);
+   mars_y_0+=     0.00021012921 * cos( 5.89022773653 +     3344.13554504880*t);
+   mars_y_0+=     0.00011370034 * cos( 3.26131408801 +     1059.38193018920*t);
+   mars_y_0+=     0.00013324177 * cos( 5.34259389724 +      529.69096509460*t);
+   mars_y_0+=     0.00006764653 * cos( 4.07671230062 +     3340.62968035200*t);
+   mars_y_0+=     0.00006764700 * cos( 2.04594066912 +     3340.59517304760*t);
+   mars_y_0+=     0.00008346220 * cos( 3.42464704002 +     2281.23049651060*t);
+   mars_y_0+=     0.00005400042 * cos( 5.81507793194 +     2942.46342329160*t);
+   mars_y_0+=     0.00004809528 * cos( 0.27875310553 +     3738.76143010800*t);
+   mars_y_0+=     0.00004849523 * cos( 2.85190987550 +     5621.84292321040*t);
+   mars_y_0+=     0.00005263268 * cos( 0.75811089992 +     6151.53388830500*t);
+   mars_y_0+=     0.00003609527 * cos( 4.53244488294 +      796.29800681640*t);
+   mars_y_0+=     0.00003724293 * cos( 5.98516013322 +    16703.06213349900*t);
+   mars_y_0+=     0.00003805073 * cos( 5.94234296399 +      398.14900340820*t);
+   mars_y_0+=     0.00002394490 * cos( 5.34678816191 +     3149.16416058820*t);
+   mars_y_0+=     0.00002251027 * cos( 0.76938193892 +     3532.06069281140*t);
+   mars_y_0+=     0.00001975769 * cos( 0.56949816579 +     6677.70173505060*t);
+   mars_y_0+=     0.00002226030 * cos( 5.99867316288 +     3340.67973700260*t);
+   mars_y_0+=     0.00002226089 * cos( 0.12398424247 +     3340.54511639700*t);
+   mars_y_0+=     0.00002177591 * cos( 0.12334436516 +     6283.07584999140*t);
+   mars_y_0+=     0.00001690439 * cos( 1.58331163985 +     5884.92684658320*t);
+   mars_y_0+=     0.00002234121 * cos( 3.24909113765 +     8962.45534991020*t);
+   mars_y_0+=     0.00001628395 * cos( 4.95250906888 +     4136.91043351620*t);
+   mars_y_0+=     0.00001956411 * cos( 6.23095843554 +     6684.74797174860*t);
+   mars_y_0+=     0.00001697214 * cos( 0.81869636263 +     2544.31441988340*t);
+   mars_y_0+=     0.00001385946 * cos( 3.73437191158 +     1751.53953141600*t);
+   mars_y_0+=     0.00001439619 * cos( 5.19505958438 +     2810.92146160520*t);
+   mars_y_0+=     0.00001281890 * cos( 3.13035275682 +     1194.44701022460*t);
+   mars_y_0+=     0.00001469783 * cos( 0.30415060688 +     3870.30339179440*t);
+   mars_y_0+=     0.00001571880 * cos( 3.27679498650 +     1589.07289528380*t);
+   mars_y_0+=     0.00001575854 * cos( 2.78266835243 +     5092.15195811580*t);
+
+   double mars_y_1=0.0;
+   mars_y_1+=     0.01427324210 * cos( 3.14159265359 +        0.00000000000*t);
+   mars_y_1+=     0.00551063753 * cos( 3.52128320402 +     6681.22485339960*t);
+   mars_y_1+=     0.00077091913 * cos( 3.86082685753 +    10021.83728009940*t);
+   mars_y_1+=     0.00037310491 * cos( 1.16016958445 +     3340.61242669980*t);
+   mars_y_1+=     0.00009582122 * cos( 4.19942354479 +    13362.44970679920*t);
+   mars_y_1+=     0.00002617695 * cos( 1.47284555520 +     3337.08930835080*t);
+   mars_y_1+=     0.00002611572 * cos( 4.65030772498 +     3344.13554504880*t);
+   mars_y_1+=     0.00001162955 * cos( 4.53778503576 +    16703.06213349900*t);
+   mars_y_1=mars_y_1 * t;
+
+   double mars_y_2=0.0;
+   mars_y_2+=     0.00035396765 * cos( 3.14159265359 +        0.00000000000*t);
+   mars_y_2+=     0.00021950759 * cos( 1.96291594946 +     6681.22485339960*t);
+   mars_y_2+=     0.00006075990 * cos( 2.35864321001 +    10021.83728009940*t);
+   mars_y_2+=     0.00002571425 * cos( 5.64795745327 +     3340.61242669980*t);
+   mars_y_2+=     0.00001129099 * cos( 2.71576248963 +    13362.44970679920*t);
+   mars_y_2=mars_y_2 * t * t;
+
+   double mars_y_3=0.0;
+   mars_y_3+=     0.00001448778 * cos( 0.00000000000 +        0.00000000000*t);
+   mars_y_3=mars_y_3 * t * t * t;
+
+   return mars_y_0+mars_y_1+mars_y_2+mars_y_3;
+}
+
+double vsop87a_milli::mars_z(double t){
+   double mars_z_0=0.0;
+   mars_z_0+=     0.04901207220 * cos( 3.76712324286 +     3340.61242669980*t);
+   mars_z_0+=     0.00660669541 * cos( 0.00000000000 +        0.00000000000*t);
+   mars_z_0+=     0.00228333904 * cos( 4.10544022266 +     6681.22485339960*t);
+   mars_z_0+=     0.00015958402 * cos( 4.44367058261 +    10021.83728009940*t);
+   mars_z_0+=     0.00001321976 * cos( 4.78186604114 +    13362.44970679920*t);
+
+   double mars_z_1=0.0;
+   mars_z_1+=     0.00331842958 * cos( 6.05027773492 +     3340.61242669980*t);
+   mars_z_1+=     0.00047930411 * cos( 3.14159265359 +        0.00000000000*t);
+   mars_z_1+=     0.00009896501 * cos( 1.61155844715 +     6681.22485339960*t);
+   mars_z_1+=     0.00001700147 * cos( 2.63703242065 +    10021.83728009940*t);
+   mars_z_1=mars_z_1 * t;
+
+   double mars_z_2=0.0;
+   mars_z_2+=     0.00013705360 * cos( 1.04212852598 +     3340.61242669980*t);
+   mars_z_2+=     0.00005931596 * cos( 3.14159265359 +        0.00000000000*t);
+   mars_z_2=mars_z_2 * t * t;
+
+   return mars_z_0+mars_z_1+mars_z_2;
+}
+
+double vsop87a_milli::mercury_x(double t){
+   double mercury_x_0=0.0;
+   mercury_x_0+=     0.37546291728 * cos( 4.39651506942 +    26087.90314157420*t);
+   mercury_x_0+=     0.03825746672 * cos( 1.16485604339 +    52175.80628314840*t);
+   mercury_x_0+=     0.02625615963 * cos( 3.14159265359 +        0.00000000000*t);
+   mercury_x_0+=     0.00584261333 * cos( 4.21599394757 +    78263.70942472259*t);
+   mercury_x_0+=     0.00105716695 * cos( 0.98379033182 +   104351.61256629678*t);
+   mercury_x_0+=     0.00021011730 * cos( 4.03469353923 +   130439.51570787099*t);
+   mercury_x_0+=     0.00004433373 * cos( 0.80236674527 +   156527.41884944518*t);
+
+   double mercury_x_1=0.0;
+   mercury_x_1+=     0.00318848034 * cos( 0.00000000000 +        0.00000000000*t);
+   mercury_x_1+=     0.00105289019 * cos( 5.91600475006 +    52175.80628314840*t);
+   mercury_x_1+=     0.00032316001 * cos( 2.68247273347 +    78263.70942472259*t);
+   mercury_x_1+=     0.00011992889 * cos( 5.81575112963 +    26087.90314157420*t);
+   mercury_x_1+=     0.00008783200 * cos( 5.73285747425 +   104351.61256629678*t);
+   mercury_x_1+=     0.00002329042 * cos( 2.50023793407 +   130439.51570787099*t);
+   mercury_x_1=mercury_x_1 * t;
+
+   double mercury_x_2=0.0;
+   mercury_x_2+=     0.00001484185 * cos( 4.35401210269 +    52175.80628314840*t);
+   mercury_x_2+=     0.00001214995 * cos( 0.00000000000 +        0.00000000000*t);
+   mercury_x_2=mercury_x_2 * t * t;
+
+   return mercury_x_0+mercury_x_1+mercury_x_2;
+}
+
+double vsop87a_milli::mercury_y(double t){
+   double mercury_y_0=0.0;
+   mercury_y_0+=     0.37953642888 * cos( 2.83780617820 +    26087.90314157420*t);
+   mercury_y_0+=     0.11626131831 * cos( 3.14159265359 +        0.00000000000*t);
+   mercury_y_0+=     0.03854668215 * cos( 5.88780608966 +    52175.80628314840*t);
+   mercury_y_0+=     0.00587711268 * cos( 2.65498896201 +    78263.70942472259*t);
+   mercury_y_0+=     0.00106235493 * cos( 5.70550616735 +   104351.61256629678*t);
+   mercury_y_0+=     0.00021100828 * cos( 2.47291315849 +   130439.51570787099*t);
+   mercury_y_0+=     0.00004450056 * cos( 5.52354907071 +   156527.41884944518*t);
+
+   double mercury_y_1=0.0;
+   mercury_y_1+=     0.00107803852 * cos( 4.34964793883 +    52175.80628314840*t);
+   mercury_y_1+=     0.00080651544 * cos( 3.14159265359 +        0.00000000000*t);
+   mercury_y_1+=     0.00032715354 * cos( 1.11763734425 +    78263.70942472259*t);
+   mercury_y_1+=     0.00008858158 * cos( 4.16852401867 +   104351.61256629678*t);
+   mercury_y_1+=     0.00011914709 * cos( 1.22139986340 +    26087.90314157420*t);
+   mercury_y_1+=     0.00002344469 * cos( 0.93615372641 +   130439.51570787099*t);
+   mercury_y_1=mercury_y_1 * t;
+
+   double mercury_y_2=0.0;
+   mercury_y_2+=     0.00004612157 * cos( 0.00000000000 +        0.00000000000*t);
+   mercury_y_2+=     0.00001575670 * cos( 2.81172733349 +    52175.80628314840*t);
+   mercury_y_2=mercury_y_2 * t * t;
+
+   return mercury_y_0+mercury_y_1+mercury_y_2;
+}
+
+double vsop87a_milli::mercury_z(double t){
+   double mercury_z_0=0.0;
+   mercury_z_0+=     0.04607665326 * cos( 1.99295081967 +    26087.90314157420*t);
+   mercury_z_0+=     0.00708734365 * cos( 3.14159265359 +        0.00000000000*t);
+   mercury_z_0+=     0.00469171617 * cos( 5.04215742764 +    52175.80628314840*t);
+   mercury_z_0+=     0.00071626395 * cos( 1.80894256071 +    78263.70942472259*t);
+   mercury_z_0+=     0.00012957446 * cos( 4.85922032010 +   104351.61256629678*t);
+   mercury_z_0+=     0.00002575002 * cos( 1.62646731545 +   130439.51570787099*t);
+
+   double mercury_z_1=0.0;
+   mercury_z_1+=     0.00108722177 * cos( 3.91134750825 +    26087.90314157420*t);
+   mercury_z_1+=     0.00057826621 * cos( 3.14159265359 +        0.00000000000*t);
+   mercury_z_1+=     0.00004297352 * cos( 2.56373047177 +    52175.80628314840*t);
+   mercury_z_1+=     0.00002435833 * cos( 0.05112640506 +    78263.70942472259*t);
+   mercury_z_1=mercury_z_1 * t;
+
+   double mercury_z_2=0.0;
+   mercury_z_2+=     0.00001053118 * cos( 5.37979214357 +    26087.90314157420*t);
+   mercury_z_2+=     0.00001185024 * cos( 0.00000000000 +        0.00000000000*t);
+   mercury_z_2=mercury_z_2 * t * t;
+
+   return mercury_z_0+mercury_z_1+mercury_z_2;
+}
+
+double vsop87a_milli::neptune_x(double t){
+   double neptune_x_0=0.0;
+   neptune_x_0+=    30.05890004476 * cos( 5.31211340029 +       38.13303563780*t);
+   neptune_x_0+=     0.27080164222 * cos( 3.14159265359 +        0.00000000000*t);
+   neptune_x_0+=     0.13505661755 * cos( 3.50078975634 +       76.26607127560*t);
+   neptune_x_0+=     0.15726094556 * cos( 0.11319072675 +       36.64856292950*t);
+   neptune_x_0+=     0.14935120126 * cos( 1.08499403018 +       39.61750834610*t);
+   neptune_x_0+=     0.02597313814 * cos( 1.99590301412 +        1.48447270830*t);
+   neptune_x_0+=     0.01074040708 * cos( 5.38502938672 +       74.78159856730*t);
+   neptune_x_0+=     0.00823793287 * cos( 1.43221581862 +       35.16409022120*t);
+   neptune_x_0+=     0.00817588813 * cos( 0.78180174031 +        2.96894541660*t);
+   neptune_x_0+=     0.00565534918 * cos( 5.98964907613 +       41.10198105440*t);
+   neptune_x_0+=     0.00495719107 * cos( 0.59948143567 +      529.69096509460*t);
+   neptune_x_0+=     0.00307525907 * cos( 0.40023311011 +       73.29712585900*t);
+   neptune_x_0+=     0.00272253551 * cos( 0.87443494387 +      213.29909543800*t);
+   neptune_x_0+=     0.00135887219 * cos( 5.54676577816 +       77.75054398390*t);
+   neptune_x_0+=     0.00090965704 * cos( 1.68910246115 +      114.39910691340*t);
+   neptune_x_0+=     0.00069040539 * cos( 5.83469123520 +        4.45341812490*t);
+   neptune_x_0+=     0.00060813556 * cos( 2.62589958380 +       33.67961751290*t);
+   neptune_x_0+=     0.00054690827 * cos( 1.55799996661 +       71.81265315070*t);
+   neptune_x_0+=     0.00028889260 * cos( 4.78966826027 +       42.58645376270*t);
+   neptune_x_0+=     0.00012614732 * cos( 3.57002516434 +      112.91463420510*t);
+   neptune_x_0+=     0.00012749153 * cos( 2.73719269645 +      111.43016149680*t);
+   neptune_x_0+=     0.00012013994 * cos( 0.94912933496 +     1059.38193018920*t);
+   neptune_x_0+=     0.00007540650 * cos( 2.77783477855 +       70.32818044240*t);
+   neptune_x_0+=     0.00007573383 * cos( 0.10011329853 +      426.59819087600*t);
+   neptune_x_0+=     0.00008004318 * cos( 1.63965626260 +      108.46121608020*t);
+   neptune_x_0+=     0.00006464842 * cos( 4.62580066013 +        5.93789083320*t);
+   neptune_x_0+=     0.00005565860 * cos( 3.82502185953 +       32.19514480460*t);
+   neptune_x_0+=     0.00004654361 * cos( 0.10385887980 +       37.61177077600*t);
+   neptune_x_0+=     0.00004732434 * cos( 4.09723977191 +       79.23501669220*t);
+   neptune_x_0+=     0.00004557247 * cos( 1.09712669317 +       38.65430049960*t);
+   neptune_x_0+=     0.00004322550 * cos( 2.37744780188 +       38.08485152800*t);
+   neptune_x_0+=     0.00004315539 * cos( 5.10473142056 +       38.18121974760*t);
+   neptune_x_0+=     0.00004089036 * cos( 1.99429048244 +       37.16982779130*t);
+   neptune_x_0+=     0.00004249674 * cos( 5.63324475823 +       28.57180808220*t);
+   neptune_x_0+=     0.00003920412 * cos( 5.49263784865 +       39.09624348430*t);
+   neptune_x_0+=     0.00003951848 * cos( 2.29996934110 +       98.89998852460*t);
+   neptune_x_0+=     0.00003322735 * cos( 4.68798591938 +        4.19278569400*t);
+   neptune_x_0+=     0.00003108292 * cos( 1.84434543409 +       47.69426319340*t);
+   neptune_x_0+=     0.00003260095 * cos( 1.81839652878 +      145.10977900970*t);
+   neptune_x_0+=     0.00002723442 * cos( 3.82296285941 +      109.94568878850*t);
+   neptune_x_0+=     0.00002522938 * cos( 4.66296126912 +      312.19908396260*t);
+   neptune_x_0+=     0.00001887430 * cos( 3.20485417792 +       35.68535508300*t);
+   neptune_x_0+=     0.00001648985 * cos( 4.06990666591 +       30.05628079050*t);
+   neptune_x_0+=     0.00001826700 * cos( 3.58024318649 +       44.07092647100*t);
+   neptune_x_0+=     0.00001945462 * cos( 4.15326825288 +      206.18554843720*t);
+   neptune_x_0+=     0.00001681255 * cos( 4.27560304282 +       40.58071619260*t);
+   neptune_x_0+=     0.00001533383 * cos( 1.17732211665 +       38.02116105320*t);
+   neptune_x_0+=     0.00001891892 * cos( 0.74998855536 +      220.41264243880*t);
+   neptune_x_0+=     0.00001527526 * cos( 0.02173640246 +       38.24491022240*t);
+   neptune_x_0+=     0.00002084477 * cos( 1.56821671369 +      149.56319713460*t);
+   neptune_x_0+=     0.00002083682 * cos( 2.83676961811 +      137.03302416240*t);
+   neptune_x_0+=     0.00001615063 * cos( 2.91063835010 +      106.97674337190*t);
+   neptune_x_0+=     0.00001265797 * cos( 3.42037275447 +       46.20979048510*t);
+   neptune_x_0+=     0.00001560429 * cos( 0.55865739143 +       37.87240320690*t);
+   neptune_x_0+=     0.00001545705 * cos( 0.64028780672 +       38.39366806870*t);
+   neptune_x_0+=     0.00001434798 * cos( 0.72658718863 +      522.57741809380*t);
+   neptune_x_0+=     0.00001271543 * cos( 2.74412981229 +       33.94024994380*t);
+   neptune_x_0+=     0.00001407422 * cos( 3.61743288666 +      536.80451209540*t);
+   neptune_x_0+=     0.00001387922 * cos( 3.71814330952 +      115.88357962170*t);
+   neptune_x_0+=     0.00001228939 * cos( 2.78878211792 +       72.07328558160*t);
+   neptune_x_0+=     0.00001448439 * cos( 1.98814317259 +      181.75834193920*t);
+   neptune_x_0+=     0.00001170078 * cos( 3.98594689041 +        8.07675484730*t);
+   neptune_x_0+=     0.00001080795 * cos( 4.75485636019 +       42.32582133180*t);
+   neptune_x_0+=     0.00001220341 * cos( 2.64791449584 +      148.07872442630*t);
+
+   double neptune_x_1=0.0;
+   neptune_x_1+=     0.00255840261 * cos( 2.01935686795 +       36.64856292950*t);
+   neptune_x_1+=     0.00243125299 * cos( 5.46214902873 +       39.61750834610*t);
+   neptune_x_1+=     0.00118398168 * cos( 2.88251845061 +       76.26607127560*t);
+   neptune_x_1+=     0.00037965449 * cos( 3.14159265359 +        0.00000000000*t);
+   neptune_x_1+=     0.00021924705 * cos( 3.20156164152 +       35.16409022120*t);
+   neptune_x_1+=     0.00017459808 * cos( 4.26349398817 +       41.10198105440*t);
+   neptune_x_1+=     0.00013130617 * cos( 5.36424961848 +        2.96894541660*t);
+   neptune_x_1+=     0.00005086527 * cos( 1.92377354729 +       38.13303563780*t);
+   neptune_x_1+=     0.00004899718 * cos( 2.09349497813 +       73.29712585900*t);
+   neptune_x_1+=     0.00002745267 * cos( 4.06252818667 +       77.75054398390*t);
+   neptune_x_1+=     0.00002204414 * cos( 4.38855639521 +       33.67961751290*t);
+   neptune_x_1+=     0.00002168719 * cos( 4.11768012563 +        4.45341812490*t);
+   neptune_x_1+=     0.00001572202 * cos( 1.07606611589 +      114.39910691340*t);
+   neptune_x_1+=     0.00001344022 * cos( 3.03802059051 +       42.58645376270*t);
+   neptune_x_1+=     0.00001285542 * cos( 6.02367554997 +       74.78159856730*t);
+   neptune_x_1=neptune_x_1 * t;
+
+   double neptune_x_2=0.0;
+   neptune_x_2+=     0.00005371138 * cos( 0.00000000000 +        0.00000000000*t);
+   neptune_x_2+=     0.00004536283 * cos( 5.02700751836 +       36.64856292950*t);
+   neptune_x_2+=     0.00004350766 * cos( 2.45420254304 +       39.61750834610*t);
+   neptune_x_2+=     0.00003092965 * cos( 0.62250463031 +       38.13303563780*t);
+   neptune_x_2+=     0.00002163703 * cos( 1.79218168368 +       76.26607127560*t);
+   neptune_x_2=neptune_x_2 * t * t;
+
+   return neptune_x_0+neptune_x_1+neptune_x_2;
+}
+
+double vsop87a_milli::neptune_y(double t){
+   double neptune_y_0=0.0;
+   neptune_y_0+=    30.06056351665 * cos( 3.74086294714 +       38.13303563780*t);
+   neptune_y_0+=     0.30205857683 * cos( 3.14159265359 +        0.00000000000*t);
+   neptune_y_0+=     0.13506391797 * cos( 1.92953034883 +       76.26607127560*t);
+   neptune_y_0+=     0.15706589373 * cos( 4.82539970129 +       36.64856292950*t);
+   neptune_y_0+=     0.14936165806 * cos( 5.79694900665 +       39.61750834610*t);
+   neptune_y_0+=     0.02584250749 * cos( 0.42549700754 +        1.48447270830*t);
+   neptune_y_0+=     0.01073739772 * cos( 3.81371728533 +       74.78159856730*t);
+   neptune_y_0+=     0.00815187583 * cos( 5.49429775826 +        2.96894541660*t);
+   neptune_y_0+=     0.00582199295 * cos( 6.19633718936 +       35.16409022120*t);
+   neptune_y_0+=     0.00565576412 * cos( 4.41843009015 +       41.10198105440*t);
+   neptune_y_0+=     0.00495581047 * cos( 5.31205825784 +      529.69096509460*t);
+   neptune_y_0+=     0.00304525203 * cos( 5.11048113661 +       73.29712585900*t);
+   neptune_y_0+=     0.00272640298 * cos( 5.58603690785 +      213.29909543800*t);
+   neptune_y_0+=     0.00135897385 * cos( 3.97553750964 +       77.75054398390*t);
+   neptune_y_0+=     0.00090970871 * cos( 0.11783619888 +      114.39910691340*t);
+   neptune_y_0+=     0.00068790261 * cos( 4.26391997151 +        4.45341812490*t);
+   neptune_y_0+=     0.00028893355 * cos( 3.21848975032 +       42.58645376270*t);
+   neptune_y_0+=     0.00020081559 * cos( 1.19787916085 +       33.67961751290*t);
+   neptune_y_0+=     0.00012613583 * cos( 1.99777332934 +      112.91463420510*t);
+   neptune_y_0+=     0.00012828708 * cos( 1.16740053443 +      111.43016149680*t);
+   neptune_y_0+=     0.00012012961 * cos( 5.66157563804 +     1059.38193018920*t);
+   neptune_y_0+=     0.00008768580 * cos( 3.23487156950 +      108.46121608020*t);
+   neptune_y_0+=     0.00007581788 * cos( 4.81169168396 +      426.59819087600*t);
+   neptune_y_0+=     0.00006439265 * cos( 3.05453259951 +        5.93789083320*t);
+   neptune_y_0+=     0.00005297978 * cos( 0.79002313990 +       71.81265315070*t);
+   neptune_y_0+=     0.00004650708 * cos( 4.81540983294 +       37.61177077600*t);
+   neptune_y_0+=     0.00004733483 * cos( 2.52620194642 +       79.23501669220*t);
+   neptune_y_0+=     0.00004557247 * cos( 5.80951552318 +       38.65430049960*t);
+   neptune_y_0+=     0.00004322550 * cos( 0.80665145881 +       38.08485152800*t);
+   neptune_y_0+=     0.00004315539 * cos( 3.53393506841 +       38.18121974760*t);
+   neptune_y_0+=     0.00004089036 * cos( 0.42349446479 +       37.16982779130*t);
+   neptune_y_0+=     0.00004247643 * cos( 4.06355336504 +       28.57180808220*t);
+   neptune_y_0+=     0.00003932515 * cos( 3.91607592815 +       39.09624348430*t);
+   neptune_y_0+=     0.00003930135 * cos( 3.86614178174 +       98.89998852460*t);
+   neptune_y_0+=     0.00003323991 * cos( 3.11674274385 +        4.19278569400*t);
+   neptune_y_0+=     0.00003112636 * cos( 0.27319642944 +       47.69426319340*t);
+   neptune_y_0+=     0.00003373281 * cos( 3.39616255650 +      145.10977900970*t);
+   neptune_y_0+=     0.00002670944 * cos( 2.31235275416 +      109.94568878850*t);
+   neptune_y_0+=     0.00002523042 * cos( 6.23400745185 +      312.19908396260*t);
+   neptune_y_0+=     0.00001888827 * cos( 1.63364331324 +       35.68535508300*t);
+   neptune_y_0+=     0.00001647474 * cos( 2.50010254963 +       30.05628079050*t);
+   neptune_y_0+=     0.00001826390 * cos( 2.00938305966 +       44.07092647100*t);
+   neptune_y_0+=     0.00001967147 * cos( 2.56634772532 +      206.18554843720*t);
+   neptune_y_0+=     0.00001681258 * cos( 2.70480318579 +       40.58071619260*t);
+   neptune_y_0+=     0.00001533383 * cos( 5.88971113590 +       38.02116105320*t);
+   neptune_y_0+=     0.00001894261 * cos( 5.46274825258 +      220.41264243880*t);
+   neptune_y_0+=     0.00001527526 * cos( 4.73412534395 +       38.24491022240*t);
+   neptune_y_0+=     0.00002086907 * cos( 6.28313624461 +      149.56319713460*t);
+   neptune_y_0+=     0.00002057794 * cos( 4.38552505781 +      137.03302416240*t);
+   neptune_y_0+=     0.00001720954 * cos( 4.49400805134 +      106.97674337190*t);
+   neptune_y_0+=     0.00001314116 * cos( 1.80386443362 +       46.20979048510*t);
+   neptune_y_0+=     0.00001732739 * cos( 4.14518500834 +       70.32818044240*t);
+   neptune_y_0+=     0.00001559193 * cos( 5.27114846878 +       37.87240320690*t);
+   neptune_y_0+=     0.00001545705 * cos( 5.35267669439 +       38.39366806870*t);
+   neptune_y_0+=     0.00001435274 * cos( 5.44292013172 +      522.57741809380*t);
+   neptune_y_0+=     0.00001404991 * cos( 2.04611088339 +      536.80451209540*t);
+   neptune_y_0+=     0.00001242929 * cos( 1.10242173566 +       33.94024994380*t);
+   neptune_y_0+=     0.00001388024 * cos( 2.14792830412 +      115.88357962170*t);
+   neptune_y_0+=     0.00001467042 * cos( 3.56226463770 +      181.75834193920*t);
+   neptune_y_0+=     0.00001227926 * cos( 1.21334651843 +       72.07328558160*t);
+   neptune_y_0+=     0.00001080807 * cos( 3.18401661435 +       42.32582133180*t);
+   neptune_y_0+=     0.00001111708 * cos( 5.51669920239 +        8.07675484730*t);
+   neptune_y_0+=     0.00001237027 * cos( 1.08622199668 +      148.07872442630*t);
+
+   double neptune_y_1=0.0;
+   neptune_y_1+=     0.00352947493 * cos( 3.14159265359 +        0.00000000000*t);
+   neptune_y_1+=     0.00256125493 * cos( 0.44757496817 +       36.64856292950*t);
+   neptune_y_1+=     0.00243147725 * cos( 3.89099798696 +       39.61750834610*t);
+   neptune_y_1+=     0.00118427205 * cos( 1.31128027037 +       76.26607127560*t);
+   neptune_y_1+=     0.00021936702 * cos( 1.63124087591 +       35.16409022120*t);
+   neptune_y_1+=     0.00017462332 * cos( 2.69229902966 +       41.10198105440*t);
+   neptune_y_1+=     0.00012992380 * cos( 3.79578633002 +        2.96894541660*t);
+   neptune_y_1+=     0.00004945117 * cos( 0.51727080684 +       73.29712585900*t);
+   neptune_y_1+=     0.00002745921 * cos( 2.49178311082 +       77.75054398390*t);
+   neptune_y_1+=     0.00002145481 * cos( 2.54768447291 +        4.45341812490*t);
+   neptune_y_1+=     0.00001572289 * cos( 5.78853350711 +      114.39910691340*t);
+   neptune_y_1+=     0.00001565725 * cos( 2.89846266272 +       33.67961751290*t);
+   neptune_y_1+=     0.00001458269 * cos( 1.61835542699 +       38.13303563780*t);
+   neptune_y_1+=     0.00001293459 * cos( 4.45868061082 +       74.78159856730*t);
+   neptune_y_1+=     0.00001343731 * cos( 1.46712622109 +       42.58645376270*t);
+   neptune_y_1=neptune_y_1 * t;
+
+   double neptune_y_2=0.0;
+   neptune_y_2+=     0.00004539421 * cos( 3.45613207922 +       36.64856292950*t);
+   neptune_y_2+=     0.00004347956 * cos( 0.88317230351 +       39.61750834610*t);
+   neptune_y_2+=     0.00003595394 * cos( 3.14159265359 +        0.00000000000*t);
+   neptune_y_2+=     0.00003058647 * cos( 5.31956613665 +       38.13303563780*t);
+   neptune_y_2+=     0.00002163809 * cos( 0.22086532214 +       76.26607127560*t);
+   neptune_y_2=neptune_y_2 * t * t;
+
+   return neptune_y_0+neptune_y_1+neptune_y_2;
+}
+
+double vsop87a_milli::neptune_z(double t){
+   double neptune_z_0=0.0;
+   neptune_z_0+=     0.92866054405 * cos( 1.44103930278 +       38.13303563780*t);
+   neptune_z_0+=     0.01245978462 * cos( 0.00000000000 +        0.00000000000*t);
+   neptune_z_0+=     0.00474333567 * cos( 2.52218774238 +       36.64856292950*t);
+   neptune_z_0+=     0.00451987936 * cos( 3.50949720541 +       39.61750834610*t);
+   neptune_z_0+=     0.00417558068 * cos( 5.91310695421 +       76.26607127560*t);
+   neptune_z_0+=     0.00084104329 * cos( 4.38928900096 +        1.48447270830*t);
+   neptune_z_0+=     0.00032704958 * cos( 1.52048692001 +       74.78159856730*t);
+   neptune_z_0+=     0.00030873335 * cos( 3.29017611456 +       35.16409022120*t);
+   neptune_z_0+=     0.00025812584 * cos( 3.19303128782 +        2.96894541660*t);
+   neptune_z_0+=     0.00016865319 * cos( 2.13251104425 +       41.10198105440*t);
+   neptune_z_0+=     0.00011789909 * cos( 3.60001877675 +      213.29909543800*t);
+   neptune_z_0+=     0.00009770125 * cos( 2.80133971586 +       73.29712585900*t);
+   neptune_z_0+=     0.00011279680 * cos( 3.55816676334 +      529.69096509460*t);
+   neptune_z_0+=     0.00004119873 * cos( 1.67934316836 +       77.75054398390*t);
+   neptune_z_0+=     0.00002818034 * cos( 4.10661077794 +      114.39910691340*t);
+   neptune_z_0+=     0.00002868677 * cos( 4.27011526203 +       33.67961751290*t);
+   neptune_z_0+=     0.00002213464 * cos( 1.96045135168 +        4.45341812490*t);
+   neptune_z_0+=     0.00001865650 * cos( 5.05540709577 +       71.81265315070*t);
+
+   double neptune_z_1=0.0;
+   neptune_z_1+=     0.00154885971 * cos( 2.14239039664 +       38.13303563780*t);
+   neptune_z_1+=     0.00007783708 * cos( 4.40146905905 +       36.64856292950*t);
+   neptune_z_1+=     0.00006862414 * cos( 1.65930160610 +       39.61750834610*t);
+   neptune_z_1+=     0.00009464276 * cos( 0.00000000000 +        0.00000000000*t);
+   neptune_z_1+=     0.00003891873 * cos( 5.46761139427 +       76.26607127560*t);
+   neptune_z_1=neptune_z_1 * t;
+
+   double neptune_z_2=0.0;
+   neptune_z_2+=     0.00001264840 * cos( 1.91401498992 +       38.13303563780*t);
+   neptune_z_2=neptune_z_2 * t * t;
+
+   return neptune_z_0+neptune_z_1+neptune_z_2;
+}
+
+double vsop87a_milli::saturn_x(double t){
+   double saturn_x_0=0.0;
+   saturn_x_0+=     9.51638335797 * cos( 0.87441380794 +      213.29909543800*t);
+   saturn_x_0+=     0.26412374238 * cos( 0.12390892620 +      426.59819087600*t);
+   saturn_x_0+=     0.06760430339 * cos( 4.16767145778 +      206.18554843720*t);
+   saturn_x_0+=     0.06624260115 * cos( 0.75094737780 +      220.41264243880*t);
+   saturn_x_0+=     0.04244797817 * cos( 0.00000000000 +        0.00000000000*t);
+   saturn_x_0+=     0.02336340488 * cos( 2.02227784673 +        7.11354700080*t);
+   saturn_x_0+=     0.01255372247 * cos( 2.17338917731 +      110.20632121940*t);
+   saturn_x_0+=     0.01115684467 * cos( 3.15686878377 +      419.48464387520*t);
+   saturn_x_0+=     0.01097683232 * cos( 5.65753337256 +      639.89728631400*t);
+   saturn_x_0+=     0.00716328481 * cos( 2.71149993708 +      316.39186965660*t);
+   saturn_x_0+=     0.00509313365 * cos( 4.95865624780 +      103.09277421860*t);
+   saturn_x_0+=     0.00433994439 * cos( 0.72012820974 +      529.69096509460*t);
+   saturn_x_0+=     0.00372894461 * cos( 0.00137195497 +      433.71173787680*t);
+   saturn_x_0+=     0.00097843523 * cos( 1.01485750417 +      323.50541665740*t);
+   saturn_x_0+=     0.00080600536 * cos( 5.62103979796 +       11.04570026390*t);
+   saturn_x_0+=     0.00083782316 * cos( 0.62038893702 +      227.52618943960*t);
+   saturn_x_0+=     0.00074150224 * cos( 2.38206066655 +      632.78373931320*t);
+   saturn_x_0+=     0.00070219382 * cos( 0.88789752415 +      209.36694217490*t);
+   saturn_x_0+=     0.00068855792 * cos( 4.01788097627 +      217.23124870110*t);
+   saturn_x_0+=     0.00065620467 * cos( 2.69728593339 +      202.25339517410*t);
+   saturn_x_0+=     0.00058297911 * cos( 2.16155251399 +      224.34479570190*t);
+   saturn_x_0+=     0.00054022837 * cos( 4.90928184374 +      853.19638175200*t);
+   saturn_x_0+=     0.00045550446 * cos( 1.88235037830 +       14.22709400160*t);
+   saturn_x_0+=     0.00038345667 * cos( 4.39815501478 +      199.07200143640*t);
+   saturn_x_0+=     0.00044551703 * cos( 5.60763553535 +       63.73589830340*t);
+   saturn_x_0+=     0.00025165185 * cos( 0.37800582257 +      216.48048917570*t);
+   saturn_x_0+=     0.00024554499 * cos( 4.53150598095 +      210.11770170030*t);
+   saturn_x_0+=     0.00024673219 * cos( 5.90891573850 +      522.57741809380*t);
+   saturn_x_0+=     0.00024677050 * cos( 5.60389382420 +      415.55249061210*t);
+   saturn_x_0+=     0.00025491374 * cos( 1.63922423181 +      117.31986822020*t);
+   saturn_x_0+=     0.00031253049 * cos( 4.62976601833 +      735.87651353180*t);
+   saturn_x_0+=     0.00023372467 * cos( 5.53491987276 +      647.01083331480*t);
+   saturn_x_0+=     0.00023355468 * cos( 0.18791490124 +      149.56319713460*t);
+   saturn_x_0+=     0.00024805815 * cos( 5.50327676733 +       74.78159856730*t);
+   saturn_x_0+=     0.00014731703 * cos( 4.67981909838 +      277.03499374140*t);
+   saturn_x_0+=     0.00012427525 * cos( 1.02995545746 +     1059.38193018920*t);
+   saturn_x_0+=     0.00009943329 * cos( 0.84628387596 +        3.93215326310*t);
+   saturn_x_0+=     0.00012393514 * cos( 4.19747622821 +      490.33408917940*t);
+   saturn_x_0+=     0.00012026472 * cos( 5.66372282839 +      351.81659230870*t);
+   saturn_x_0+=     0.00008222014 * cos( 2.47875301104 +      742.99006053260*t);
+   saturn_x_0+=     0.00009087093 * cos( 4.33505326762 +     1052.26838318840*t);
+   saturn_x_0+=     0.00006717741 * cos( 5.51897460997 +      838.96928775040*t);
+   saturn_x_0+=     0.00006232999 * cos( 2.45837758015 +      846.08283475120*t);
+   saturn_x_0+=     0.00007161671 * cos( 2.18152751738 +       95.97922721780*t);
+   saturn_x_0+=     0.00006321101 * cos( 0.83915408770 +      309.27832265580*t);
+   saturn_x_0+=     0.00006074958 * cos( 6.15905897331 +      440.82528487760*t);
+   saturn_x_0+=     0.00005343894 * cos( 3.60046273598 +      412.37109687440*t);
+   saturn_x_0+=     0.00004860582 * cos( 0.26461045175 +      536.80451209540*t);
+   saturn_x_0+=     0.00005775802 * cos( 5.30717695229 +       38.13303563780*t);
+   saturn_x_0+=     0.00005194178 * cos( 4.54584467686 +      210.85141488320*t);
+   saturn_x_0+=     0.00005152474 * cos( 0.34669517150 +      215.74677599280*t);
+   saturn_x_0+=     0.00003792540 * cos( 5.99766568983 +      422.66603761290*t);
+   saturn_x_0+=     0.00003762834 * cos( 3.72112920226 +      212.33588759150*t);
+   saturn_x_0+=     0.00003747433 * cos( 1.16965137714 +      214.26230328450*t);
+   saturn_x_0+=     0.00003114576 * cos( 0.84631897292 +      213.25091132820*t);
+   saturn_x_0+=     0.00003113641 * cos( 4.04410367190 +      213.34727954780*t);
+   saturn_x_0+=     0.00002990421 * cos( 0.04148806852 +      625.67019231240*t);
+   saturn_x_0+=     0.00004111695 * cos( 5.96153153046 +      137.03302416240*t);
+   saturn_x_0+=     0.00002966450 * cos( 5.39568820046 +      138.51749687070*t);
+   saturn_x_0+=     0.00002827527 * cos( 0.73252555642 +      330.61896365820*t);
+   saturn_x_0+=     0.00003363323 * cos( 1.42089586686 +      437.64389113990*t);
+   saturn_x_0+=     0.00002886599 * cos( 1.14057922619 +       85.82729883120*t);
+   saturn_x_0+=     0.00002634075 * cos( 5.40645201521 +      288.08069400530*t);
+   saturn_x_0+=     0.00002713354 * cos( 0.96812639712 +      203.73786788240*t);
+   saturn_x_0+=     0.00003169390 * cos( 5.76640408988 +       76.26607127560*t);
+   saturn_x_0+=     0.00002618634 * cos( 5.49334837098 +      127.47179660680*t);
+   saturn_x_0+=     0.00002527746 * cos( 5.09752068381 +      628.85158605010*t);
+   saturn_x_0+=     0.00002989778 * cos( 4.15673836604 +     1066.49547719000*t);
+   saturn_x_0+=     0.00002507415 * cos( 1.49447138038 +        9.56122755560*t);
+   saturn_x_0+=     0.00002470181 * cos( 5.27435870056 +     1155.36115740700*t);
+   saturn_x_0+=     0.00002427626 * cos( 3.97311214231 +      222.86032299360*t);
+   saturn_x_0+=     0.00003128325 * cos( 4.05483976553 +     1368.66025284500*t);
+   saturn_x_0+=     0.00002309076 * cos( 3.67821438247 +      430.53034413910*t);
+   saturn_x_0+=     0.00002162629 * cos( 3.26951119901 +      340.77089204480*t);
+   saturn_x_0+=     0.00002912676 * cos( 1.76893577106 +        3.18139373770*t);
+   saturn_x_0+=     0.00002095366 * cos( 3.55759089756 +      423.41679713830*t);
+   saturn_x_0+=     0.00002335270 * cos( 5.86791072516 +      388.46515523820*t);
+   saturn_x_0+=     0.00001634262 * cos( 4.54357767539 +       12.53017297220*t);
+   saturn_x_0+=     0.00001635975 * cos( 2.19968869780 +      212.77783057620*t);
+   saturn_x_0+=     0.00001632759 * cos( 2.69164822165 +      213.82036029980*t);
+   saturn_x_0+=     0.00001498689 * cos( 3.60168057129 +       52.69019803950*t);
+   saturn_x_0+=     0.00001461217 * cos( 5.92456743836 +      429.77958461370*t);
+   saturn_x_0+=     0.00001963947 * cos( 2.05086487180 +        1.48447270830*t);
+   saturn_x_0+=     0.00001538425 * cos( 4.78544077085 +      860.30992875280*t);
+   saturn_x_0+=     0.00001485856 * cos( 5.65501463408 +      949.17560896980*t);
+   saturn_x_0+=     0.00001418116 * cos( 5.41419993599 +      350.33211960040*t);
+   saturn_x_0+=     0.00001147607 * cos( 0.19147238521 +      942.06206196900*t);
+   saturn_x_0+=     0.00001111703 * cos( 0.47907488492 +      234.63973644040*t);
+   saturn_x_0+=     0.00001019566 * cos( 5.00707811029 +     1471.75302706360*t);
+   saturn_x_0+=     0.00001009190 * cos( 1.34289487761 +      265.98929347750*t);
+   saturn_x_0+=     0.00001089450 * cos( 5.82690672710 +      362.86229257260*t);
+   saturn_x_0+=     0.00001092244 * cos( 1.13561107749 +      173.94221952280*t);
+   saturn_x_0+=     0.00001173456 * cos( 3.79591687208 +     1685.05212250160*t);
+   saturn_x_0+=     0.00001118369 * cos( 3.46624149583 +      703.63318461740*t);
+   saturn_x_0+=     0.00001150595 * cos( 3.74707160019 +      200.76892246580*t);
+   saturn_x_0+=     0.00001007761 * cos( 1.08964371328 +      225.82926841020*t);
+
+   double saturn_x_1=0.0;
+   saturn_x_1+=     0.07575103962 * cos( 0.00000000000 +        0.00000000000*t);
+   saturn_x_1+=     0.03085041716 * cos( 4.27565749128 +      426.59819087600*t);
+   saturn_x_1+=     0.02714918399 * cos( 5.85229412397 +      206.18554843720*t);
+   saturn_x_1+=     0.02643100909 * cos( 5.33291950584 +      220.41264243880*t);
+   saturn_x_1+=     0.00627104520 * cos( 0.32898307969 +        7.11354700080*t);
+   saturn_x_1+=     0.00256560953 * cos( 3.52478934343 +      639.89728631400*t);
+   saturn_x_1+=     0.00312356512 * cos( 4.83001724941 +      419.48464387520*t);
+   saturn_x_1+=     0.00189196274 * cos( 4.48642453552 +      433.71173787680*t);
+   saturn_x_1+=     0.00203646570 * cos( 1.10998681782 +      213.29909543800*t);
+   saturn_x_1+=     0.00119531145 * cos( 1.14735096078 +      110.20632121940*t);
+   saturn_x_1+=     0.00066764238 * cos( 3.72346596928 +      316.39186965660*t);
+   saturn_x_1+=     0.00066901225 * cos( 5.20257500380 +      227.52618943960*t);
+   saturn_x_1+=     0.00031000840 * cos( 6.06067919437 +      199.07200143640*t);
+   saturn_x_1+=     0.00030418100 * cos( 0.18746903351 +       14.22709400160*t);
+   saturn_x_1+=     0.00022275210 * cos( 6.19530878014 +      103.09277421860*t);
+   saturn_x_1+=     0.00018939377 * cos( 2.77618306725 +      853.19638175200*t);
+   saturn_x_1+=     0.00018093009 * cos( 5.09162723865 +      209.36694217490*t);
+   saturn_x_1+=     0.00017777854 * cos( 6.10381593351 +      217.23124870110*t);
+   saturn_x_1+=     0.00016296201 * cos( 4.86945681437 +      216.48048917570*t);
+   saturn_x_1+=     0.00017120250 * cos( 4.59611664188 +      632.78373931320*t);
+   saturn_x_1+=     0.00015894491 * cos( 0.03653502304 +      210.11770170030*t);
+   saturn_x_1+=     0.00016192653 * cos( 5.60798014450 +      323.50541665740*t);
+   saturn_x_1+=     0.00014466010 * cos( 3.67449380090 +      647.01083331480*t);
+   saturn_x_1+=     0.00011061528 * cos( 0.03163071461 +      117.31986822020*t);
+   saturn_x_1+=     0.00009873183 * cos( 5.20065307357 +      202.25339517410*t);
+   saturn_x_1+=     0.00008707608 * cos( 6.03511731637 +      224.34479570190*t);
+   saturn_x_1+=     0.00005499109 * cos( 4.40350603415 +      440.82528487760*t);
+   saturn_x_1+=     0.00005512222 * cos( 2.60556642348 +       11.04570026390*t);
+   saturn_x_1+=     0.00004008257 * cos( 1.48942966807 +      522.57741809380*t);
+   saturn_x_1+=     0.00003571196 * cos( 5.10821908379 +      412.37109687440*t);
+   saturn_x_1+=     0.00002731381 * cos( 4.10892223660 +      149.56319713460*t);
+   saturn_x_1+=     0.00002763786 * cos( 3.96253590209 +       95.97922721780*t);
+   saturn_x_1+=     0.00001875862 * cos( 2.52384080586 +        3.93215326310*t);
+   saturn_x_1+=     0.00001765816 * cos( 0.75684544353 +      277.03499374140*t);
+   saturn_x_1+=     0.00001688957 * cos( 3.98270950731 +      422.66603761290*t);
+   saturn_x_1+=     0.00001544787 * cos( 5.30283923836 +      330.61896365820*t);
+   saturn_x_1+=     0.00001449143 * cos( 6.26507179861 +      529.69096509460*t);
+   saturn_x_1+=     0.00001395551 * cos( 2.03533642541 +     1066.49547719000*t);
+   saturn_x_1+=     0.00001330413 * cos( 5.06312203212 +      234.63973644040*t);
+   saturn_x_1+=     0.00001511648 * cos( 6.23274598777 +        3.18139373770*t);
+   saturn_x_1+=     0.00001271050 * cos( 2.40338468675 +      415.55249061210*t);
+   saturn_x_1+=     0.00001171680 * cos( 4.59341412127 +      536.80451209540*t);
+   saturn_x_1+=     0.00001129056 * cos( 5.45794529295 +      423.41679713830*t);
+   saturn_x_1+=     0.00001126400 * cos( 2.88173213734 +      860.30992875280*t);
+   saturn_x_1+=     0.00001101245 * cos( 4.07698108824 +      429.77958461370*t);
+   saturn_x_1+=     0.00001273315 * cos( 0.09572429396 +      742.99006053260*t);
+   saturn_x_1+=     0.00001342277 * cos( 2.98929557875 +      210.85141488320*t);
+   saturn_x_1+=     0.00001331045 * cos( 1.90899526877 +      215.74677599280*t);
+   saturn_x_1=saturn_x_1 * t;
+
+   double saturn_x_2=0.0;
+   saturn_x_2+=     0.00560746334 * cos( 1.26401632282 +      206.18554843720*t);
+   saturn_x_2+=     0.00545834518 * cos( 3.62343709657 +      220.41264243880*t);
+   saturn_x_2+=     0.00443342186 * cos( 3.14159265359 +        0.00000000000*t);
+   saturn_x_2+=     0.00336109713 * cos( 2.42547432460 +      213.29909543800*t);
+   saturn_x_2+=     0.00224302269 * cos( 2.49151203519 +      426.59819087600*t);
+   saturn_x_2+=     0.00087170924 * cos( 4.89048951691 +        7.11354700080*t);
+   saturn_x_2+=     0.00050028094 * cos( 2.70119046081 +      433.71173787680*t);
+   saturn_x_2+=     0.00045122590 * cos( 0.36735068943 +      419.48464387520*t);
+   saturn_x_2+=     0.00032847824 * cos( 1.59210153669 +      639.89728631400*t);
+   saturn_x_2+=     0.00027153555 * cos( 3.49804002218 +      227.52618943960*t);
+   saturn_x_2+=     0.00012676167 * cos( 1.45465729530 +      199.07200143640*t);
+   saturn_x_2+=     0.00010330738 * cos( 4.76949531290 +       14.22709400160*t);
+   saturn_x_2+=     0.00007249149 * cos( 5.70264553247 +      110.20632121940*t);
+   saturn_x_2+=     0.00004653214 * cos( 1.83710048213 +      647.01083331480*t);
+   saturn_x_2+=     0.00004923585 * cos( 3.08463039042 +      216.48048917570*t);
+   saturn_x_2+=     0.00004777358 * cos( 1.81695155349 +      210.11770170030*t);
+   saturn_x_2+=     0.00004166633 * cos( 5.32887874226 +      316.39186965660*t);
+   saturn_x_2+=     0.00003508385 * cos( 0.78251653369 +      853.19638175200*t);
+   saturn_x_2+=     0.00002660470 * cos( 3.16731393212 +      209.36694217490*t);
+   saturn_x_2+=     0.00002538027 * cos( 2.65097612407 +      440.82528487760*t);
+   saturn_x_2+=     0.00002568114 * cos( 1.74024228572 +      217.23124870110*t);
+   saturn_x_2+=     0.00002503277 * cos( 4.69450368911 +      117.31986822020*t);
+   saturn_x_2+=     0.00002506986 * cos( 1.74781817701 +      103.09277421860*t);
+   saturn_x_2+=     0.00002129256 * cos( 0.28453141367 +      632.78373931320*t);
+   saturn_x_2+=     0.00001841989 * cos( 3.99269872894 +      323.50541665740*t);
+   saturn_x_2+=     0.00001228511 * cos( 0.42906039519 +      412.37109687440*t);
+   saturn_x_2=saturn_x_2 * t * t;
+
+   double saturn_x_3=0.0;
+   saturn_x_3+=     0.00077115952 * cos( 2.97714385362 +      206.18554843720*t);
+   saturn_x_3+=     0.00075340436 * cos( 1.89208005248 +      220.41264243880*t);
+   saturn_x_3+=     0.00018450895 * cos( 3.14159265359 +        0.00000000000*t);
+   saturn_x_3+=     0.00010527244 * cos( 0.66368256891 +      426.59819087600*t);
+   saturn_x_3+=     0.00008994946 * cos( 0.91696559755 +      433.71173787680*t);
+   saturn_x_3+=     0.00007403594 * cos( 1.78627385870 +      227.52618943960*t);
+   saturn_x_3+=     0.00008045160 * cos( 3.12864412887 +        7.11354700080*t);
+   saturn_x_3+=     0.00004505149 * cos( 2.24531319187 +      419.48464387520*t);
+   saturn_x_3+=     0.00003468010 * cos( 3.14590544446 +      199.07200143640*t);
+   saturn_x_3+=     0.00002974601 * cos( 6.00030641555 +      639.89728631400*t);
+   saturn_x_3+=     0.00002342089 * cos( 3.06091771643 +       14.22709400160*t);
+   saturn_x_3+=     0.00001230715 * cos( 4.38196130069 +      213.29909543800*t);
+   saturn_x_3+=     0.00001023888 * cos( 0.01138655869 +      647.01083331480*t);
+   saturn_x_3=saturn_x_3 * t * t * t;
+
+   double saturn_x_4=0.0;
+   saturn_x_4+=     0.00007959921 * cos( 4.70523623364 +      206.18554843720*t);
+   saturn_x_4+=     0.00007836652 * cos( 0.13981693631 +      220.41264243880*t);
+   saturn_x_4+=     0.00001511196 * cos( 0.06561560462 +      227.52618943960*t);
+   saturn_x_4+=     0.00001223066 * cos( 5.41618485361 +      433.71173787680*t);
+   saturn_x_4=saturn_x_4 * t * t * t * t;
+
+   return saturn_x_0+saturn_x_1+saturn_x_2+saturn_x_3+saturn_x_4;
+}
+
+double vsop87a_milli::saturn_y(double t){
+   double saturn_y_0=0.0;
+   saturn_y_0+=     9.52986882699 * cos( 5.58600556665 +      213.29909543800*t);
+   saturn_y_0+=     0.79387988806 * cos( 3.14159265359 +        0.00000000000*t);
+   saturn_y_0+=     0.26441781302 * cos( 4.83528061849 +      426.59819087600*t);
+   saturn_y_0+=     0.06916653915 * cos( 2.55279408706 +      206.18554843720*t);
+   saturn_y_0+=     0.06633570703 * cos( 5.46258848288 +      220.41264243880*t);
+   saturn_y_0+=     0.02345609742 * cos( 0.44652132519 +        7.11354700080*t);
+   saturn_y_0+=     0.01183874652 * cos( 1.34638298371 +      419.48464387520*t);
+   saturn_y_0+=     0.01245790434 * cos( 0.60367177975 +      110.20632121940*t);
+   saturn_y_0+=     0.01098751131 * cos( 4.08608782813 +      639.89728631400*t);
+   saturn_y_0+=     0.00700849336 * cos( 1.13611298025 +      316.39186965660*t);
+   saturn_y_0+=     0.00434466176 * cos( 5.42474696262 +      529.69096509460*t);
+   saturn_y_0+=     0.00373327342 * cos( 4.71308726958 +      433.71173787680*t);
+   saturn_y_0+=     0.00335162363 * cos( 0.66422253983 +      103.09277421860*t);
+   saturn_y_0+=     0.00097837745 * cos( 5.72844290173 +      323.50541665740*t);
+   saturn_y_0+=     0.00080571808 * cos( 4.05295449910 +       11.04570026390*t);
+   saturn_y_0+=     0.00083899691 * cos( 5.33204070267 +      227.52618943960*t);
+   saturn_y_0+=     0.00070158491 * cos( 5.59777963629 +      209.36694217490*t);
+   saturn_y_0+=     0.00065937657 * cos( 1.25969608208 +      202.25339517410*t);
+   saturn_y_0+=     0.00070957225 * cos( 0.88888207567 +      632.78373931320*t);
+   saturn_y_0+=     0.00068985859 * cos( 2.44460312617 +      217.23124870110*t);
+   saturn_y_0+=     0.00058382264 * cos( 0.58978766922 +      224.34479570190*t);
+   saturn_y_0+=     0.00054049836 * cos( 3.33757904879 +      853.19638175200*t);
+   saturn_y_0+=     0.00045790930 * cos( 0.30331527632 +       14.22709400160*t);
+   saturn_y_0+=     0.00041976402 * cos( 2.62591355948 +      199.07200143640*t);
+   saturn_y_0+=     0.00044697175 * cos( 0.90661238256 +       63.73589830340*t);
+   saturn_y_0+=     0.00025199575 * cos( 5.08963506006 +      216.48048917570*t);
+   saturn_y_0+=     0.00024640836 * cos( 2.95445247282 +      210.11770170030*t);
+   saturn_y_0+=     0.00024835151 * cos( 4.02630190571 +      415.55249061210*t);
+   saturn_y_0+=     0.00025545907 * cos( 0.06626229252 +      117.31986822020*t);
+   saturn_y_0+=     0.00029666833 * cos( 6.09910638345 +      735.87651353180*t);
+   saturn_y_0+=     0.00023396742 * cos( 3.96337393635 +      647.01083331480*t);
+   saturn_y_0+=     0.00023380691 * cos( 4.90051072276 +      149.56319713460*t);
+   saturn_y_0+=     0.00020272215 * cos( 2.34319548198 +      309.27832265580*t);
+   saturn_y_0+=     0.00020099552 * cos( 0.98365186365 +      522.57741809380*t);
+   saturn_y_0+=     0.00024827950 * cos( 3.92681428900 +       74.78159856730*t);
+   saturn_y_0+=     0.00015383927 * cos( 3.10227822627 +      277.03499374140*t);
+   saturn_y_0+=     0.00011629210 * cos( 5.74108283772 +     1059.38193018920*t);
+   saturn_y_0+=     0.00012422966 * cos( 2.62557865743 +      490.33408917940*t);
+   saturn_y_0+=     0.00012048048 * cos( 4.09265980116 +      351.81659230870*t);
+   saturn_y_0+=     0.00009551796 * cos( 3.48788042094 +       95.97922721780*t);
+   saturn_y_0+=     0.00007670379 * cos( 1.16594276164 +      742.99006053260*t);
+   saturn_y_0+=     0.00006919946 * cos( 1.17090063883 +      412.37109687440*t);
+   saturn_y_0+=     0.00009034877 * cos( 5.86816144198 +     1052.26838318840*t);
+   saturn_y_0+=     0.00006536751 * cos( 0.84246459392 +      838.96928775040*t);
+   saturn_y_0+=     0.00006082097 * cos( 4.58758280729 +      440.82528487760*t);
+   saturn_y_0+=     0.00005027211 * cos( 0.93213690546 +      846.08283475120*t);
+   saturn_y_0+=     0.00004838146 * cos( 4.98563812475 +      536.80451209540*t);
+   saturn_y_0+=     0.00005768897 * cos( 3.73776690402 +       38.13303563780*t);
+   saturn_y_0+=     0.00005201849 * cos( 2.97482802430 +      210.85141488320*t);
+   saturn_y_0+=     0.00005156578 * cos( 5.05796998564 +      215.74677599280*t);
+   saturn_y_0+=     0.00003792348 * cos( 4.41806046981 +      422.66603761290*t);
+   saturn_y_0+=     0.00003881104 * cos( 3.38026646963 +        3.93215326310*t);
+   saturn_y_0+=     0.00003768751 * cos( 2.14954247360 +      212.33588759150*t);
+   saturn_y_0+=     0.00003752010 * cos( 5.88125434018 +      214.26230328450*t);
+   saturn_y_0+=     0.00003118938 * cos( 5.55799397159 +      213.25091132820*t);
+   saturn_y_0+=     0.00003118056 * cos( 2.47259780102 +      213.34727954780*t);
+   saturn_y_0+=     0.00003815691 * cos( 2.44166851155 +      625.67019231240*t);
+   saturn_y_0+=     0.00004042463 * cos( 1.24471211016 +      137.03302416240*t);
+   saturn_y_0+=     0.00002829944 * cos( 5.44434225998 +      330.61896365820*t);
+   saturn_y_0+=     0.00003367270 * cos( 6.13298847057 +      437.64389113990*t);
+   saturn_y_0+=     0.00002891352 * cos( 5.85313497106 +       85.82729883120*t);
+   saturn_y_0+=     0.00002774595 * cos( 5.61594351302 +      203.73786788240*t);
+   saturn_y_0+=     0.00002638715 * cos( 3.83678156812 +      288.08069400530*t);
+   saturn_y_0+=     0.00003171680 * cos( 4.19553075395 +       76.26607127560*t);
+   saturn_y_0+=     0.00002532374 * cos( 3.52629372341 +      628.85158605010*t);
+   saturn_y_0+=     0.00002533632 * cos( 3.89788590926 +      138.51749687070*t);
+   saturn_y_0+=     0.00002982174 * cos( 2.58535107213 +     1066.49547719000*t);
+   saturn_y_0+=     0.00002620642 * cos( 0.69751279148 +      127.47179660680*t);
+   saturn_y_0+=     0.00002501775 * cos( 6.19929274396 +        9.56122755560*t);
+   saturn_y_0+=     0.00002448467 * cos( 0.54179432209 +     1155.36115740700*t);
+   saturn_y_0+=     0.00002431496 * cos( 2.40122451395 +      222.86032299360*t);
+   saturn_y_0+=     0.00003138628 * cos( 5.63058455924 +     1368.66025284500*t);
+   saturn_y_0+=     0.00002269226 * cos( 2.12401905105 +      430.53034413910*t);
+   saturn_y_0+=     0.00002078049 * cos( 1.95682348964 +      423.41679713830*t);
+   saturn_y_0+=     0.00002670750 * cos( 0.18165311734 +        3.18139373770*t);
+   saturn_y_0+=     0.00002339764 * cos( 4.29619053852 +      388.46515523820*t);
+   saturn_y_0+=     0.00001636580 * cos( 2.97440139727 +       12.53017297220*t);
+   saturn_y_0+=     0.00001639181 * cos( 0.62823227849 +      212.77783057620*t);
+   saturn_y_0+=     0.00001634235 * cos( 1.12043073218 +      213.82036029980*t);
+   saturn_y_0+=     0.00001499665 * cos( 5.16865990579 +       52.69019803950*t);
+   saturn_y_0+=     0.00001462908 * cos( 4.35285690993 +      429.77958461370*t);
+   saturn_y_0+=     0.00001678952 * cos( 2.07211719214 +      949.17560896980*t);
+   saturn_y_0+=     0.00001989253 * cos( 0.42496478369 +        1.48447270830*t);
+   saturn_y_0+=     0.00001540242 * cos( 3.21449770483 +      860.30992875280*t);
+   saturn_y_0+=     0.00001437063 * cos( 3.84293543293 +      350.33211960040*t);
+   saturn_y_0+=     0.00001276377 * cos( 2.98728987770 +      340.77089204480*t);
+   saturn_y_0+=     0.00001152164 * cos( 1.78736848302 +      942.06206196900*t);
+   saturn_y_0+=     0.00001112617 * cos( 5.19114183145 +      234.63973644040*t);
+   saturn_y_0+=     0.00001011023 * cos( 0.27242160432 +     1471.75302706360*t);
+   saturn_y_0+=     0.00001162807 * cos( 2.35040840317 +      200.76892246580*t);
+   saturn_y_0+=     0.00001091025 * cos( 4.25638370205 +      362.86229257260*t);
+   saturn_y_0+=     0.00001090678 * cos( 5.85086226218 +      173.94221952280*t);
+   saturn_y_0+=     0.00001186647 * cos( 5.38323620554 +     1685.05212250160*t);
+   saturn_y_0+=     0.00001120819 * cos( 1.89478696683 +      703.63318461740*t);
+   saturn_y_0+=     0.00001008898 * cos( 5.80110302450 +      225.82926841020*t);
+
+   double saturn_y_1=0.0;
+   saturn_y_1+=     0.05373889135 * cos( 0.00000000000 +        0.00000000000*t);
+   saturn_y_1+=     0.03090575152 * cos( 2.70346890906 +      426.59819087600*t);
+   saturn_y_1+=     0.02741594312 * cos( 4.26667636015 +      206.18554843720*t);
+   saturn_y_1+=     0.02647489677 * cos( 3.76132298889 +      220.41264243880*t);
+   saturn_y_1+=     0.00631520527 * cos( 5.03245505280 +        7.11354700080*t);
+   saturn_y_1+=     0.00256799701 * cos( 1.95351819758 +      639.89728631400*t);
+   saturn_y_1+=     0.00312271930 * cos( 3.25850205023 +      419.48464387520*t);
+   saturn_y_1+=     0.00189433319 * cos( 2.91501840819 +      433.71173787680*t);
+   saturn_y_1+=     0.00164133553 * cos( 5.29239290066 +      213.29909543800*t);
+   saturn_y_1+=     0.00116791227 * cos( 5.89146675760 +      110.20632121940*t);
+   saturn_y_1+=     0.00067210919 * cos( 2.17042636344 +      316.39186965660*t);
+   saturn_y_1+=     0.00067003292 * cos( 3.63101075514 +      227.52618943960*t);
+   saturn_y_1+=     0.00033002406 * cos( 4.35527405801 +      199.07200143640*t);
+   saturn_y_1+=     0.00030628998 * cos( 4.88861760772 +       14.22709400160*t);
+   saturn_y_1+=     0.00022234714 * cos( 4.62212779231 +      103.09277421860*t);
+   saturn_y_1+=     0.00018945004 * cos( 1.20412493845 +      853.19638175200*t);
+   saturn_y_1+=     0.00018079959 * cos( 3.51566153251 +      209.36694217490*t);
+   saturn_y_1+=     0.00017791543 * cos( 4.53214140649 +      217.23124870110*t);
+   saturn_y_1+=     0.00016320701 * cos( 3.29784030970 +      216.48048917570*t);
+   saturn_y_1+=     0.00015944258 * cos( 4.74503265169 +      210.11770170030*t);
+   saturn_y_1+=     0.00016717122 * cos( 3.00270792752 +      632.78373931320*t);
+   saturn_y_1+=     0.00016149947 * cos( 4.04186432517 +      323.50541665740*t);
+   saturn_y_1+=     0.00014481431 * cos( 2.10298298650 +      647.01083331480*t);
+   saturn_y_1+=     0.00011084040 * cos( 4.74073871754 +      117.31986822020*t);
+   saturn_y_1+=     0.00009905491 * cos( 3.60258599375 +      202.25339517410*t);
+   saturn_y_1+=     0.00008726051 * cos( 4.46341342877 +      224.34479570190*t);
+   saturn_y_1+=     0.00006585597 * cos( 4.07326320487 +      309.27832265580*t);
+   saturn_y_1+=     0.00005505978 * cos( 2.83207390240 +      440.82528487760*t);
+   saturn_y_1+=     0.00005424041 * cos( 1.03197684410 +       11.04570026390*t);
+   saturn_y_1+=     0.00004178266 * cos( 3.01038512076 +      412.37109687440*t);
+   saturn_y_1+=     0.00004049905 * cos( 5.17488767645 +       95.97922721780*t);
+   saturn_y_1+=     0.00002735256 * cos( 2.53975850409 +      149.56319713460*t);
+   saturn_y_1+=     0.00002369024 * cos( 2.40497927917 +      522.57741809380*t);
+   saturn_y_1+=     0.00001745258 * cos( 5.50576015456 +      277.03499374140*t);
+   saturn_y_1+=     0.00001692790 * cos( 2.39926502529 +      422.66603761290*t);
+   saturn_y_1+=     0.00001546006 * cos( 3.73156925599 +      330.61896365820*t);
+   saturn_y_1+=     0.00001389354 * cos( 0.46207025895 +     1066.49547719000*t);
+   saturn_y_1+=     0.00001332553 * cos( 3.49199812296 +      234.63973644040*t);
+   saturn_y_1+=     0.00001300934 * cos( 0.83727681906 +      415.55249061210*t);
+   saturn_y_1+=     0.00001393622 * cos( 4.62214277175 +        3.18139373770*t);
+   saturn_y_1+=     0.00001174319 * cos( 2.72609984335 +      846.08283475120*t);
+   saturn_y_1+=     0.00001148010 * cos( 3.04374738882 +      536.80451209540*t);
+   saturn_y_1+=     0.00001314125 * cos( 4.38891656600 +      625.67019231240*t);
+   saturn_y_1+=     0.00001127952 * cos( 3.87309692307 +      423.41679713830*t);
+   saturn_y_1+=     0.00001127646 * cos( 1.31088906213 +      860.30992875280*t);
+   saturn_y_1+=     0.00001519732 * cos( 6.12880664637 +        3.93215326310*t);
+   saturn_y_1+=     0.00001102361 * cos( 2.50535306014 +      429.77958461370*t);
+   saturn_y_1+=     0.00001344891 * cos( 1.41793593685 +      210.85141488320*t);
+   saturn_y_1+=     0.00001331786 * cos( 0.33834520814 +      215.74677599280*t);
+   saturn_y_1=saturn_y_1 * t;
+
+   double saturn_y_2=0.0;
+   saturn_y_2+=     0.00563706537 * cos( 5.97115878242 +      206.18554843720*t);
+   saturn_y_2+=     0.00547012116 * cos( 2.05154973426 +      220.41264243880*t);
+   saturn_y_2+=     0.00458518613 * cos( 0.00000000000 +        0.00000000000*t);
+   saturn_y_2+=     0.00362294249 * cos( 0.89540100509 +      213.29909543800*t);
+   saturn_y_2+=     0.00225521642 * cos( 0.91699821445 +      426.59819087600*t);
+   saturn_y_2+=     0.00088390611 * cos( 3.30289449917 +        7.11354700080*t);
+   saturn_y_2+=     0.00050101314 * cos( 1.12976163835 +      433.71173787680*t);
+   saturn_y_2+=     0.00045516403 * cos( 5.07669466539 +      419.48464387520*t);
+   saturn_y_2+=     0.00032896745 * cos( 0.02089057938 +      639.89728631400*t);
+   saturn_y_2+=     0.00027199743 * cos( 1.92638417640 +      227.52618943960*t);
+   saturn_y_2+=     0.00013251505 * cos( 6.07693099404 +      199.07200143640*t);
+   saturn_y_2+=     0.00010425984 * cos( 3.18246869028 +       14.22709400160*t);
+   saturn_y_2+=     0.00006673556 * cos( 4.24747633887 +      110.20632121940*t);
+   saturn_y_2+=     0.00004658591 * cos( 0.26557833758 +      647.01083331480*t);
+   saturn_y_2+=     0.00004934094 * cos( 1.51301179516 +      216.48048917570*t);
+   saturn_y_2+=     0.00004789554 * cos( 0.24337901916 +      210.11770170030*t);
+   saturn_y_2+=     0.00004167268 * cos( 3.73203671391 +      316.39186965660*t);
+   saturn_y_2+=     0.00003509537 * cos( 5.49281440568 +      853.19638175200*t);
+   saturn_y_2+=     0.00002743470 * cos( 6.21939083886 +      103.09277421860*t);
+   saturn_y_2+=     0.00002661172 * cos( 1.58795412736 +      209.36694217490*t);
+   saturn_y_2+=     0.00002541191 * cos( 1.07964653574 +      440.82528487760*t);
+   saturn_y_2+=     0.00002568018 * cos( 0.16811216098 +      217.23124870110*t);
+   saturn_y_2+=     0.00002507738 * cos( 3.11882746290 +      117.31986822020*t);
+   saturn_y_2+=     0.00002159089 * cos( 4.99912567024 +      632.78373931320*t);
+   saturn_y_2+=     0.00001828412 * cos( 2.43368650590 +      323.50541665740*t);
+   saturn_y_2+=     0.00001351629 * cos( 4.81673889364 +      412.37109687440*t);
+   saturn_y_2+=     0.00001177305 * cos( 5.84484412189 +      309.27832265580*t);
+   saturn_y_2=saturn_y_2 * t * t;
+
+   double saturn_y_3=0.0;
+   saturn_y_3+=     0.00077376615 * cos( 1.40391048961 +      206.18554843720*t);
+   saturn_y_3+=     0.00075564351 * cos( 0.31962896379 +      220.41264243880*t);
+   saturn_y_3+=     0.00022843837 * cos( 3.14159265359 +        0.00000000000*t);
+   saturn_y_3+=     0.00010672263 * cos( 5.36495663820 +      426.59819087600*t);
+   saturn_y_3+=     0.00009010175 * cos( 5.62865146645 +      433.71173787680*t);
+   saturn_y_3+=     0.00007418018 * cos( 0.21442310101 +      227.52618943960*t);
+   saturn_y_3+=     0.00008298723 * cos( 1.52262563519 +        7.11354700080*t);
+   saturn_y_3+=     0.00004507061 * cos( 0.67248969480 +      419.48464387520*t);
+   saturn_y_3+=     0.00003581682 * cos( 1.51466786030 +      199.07200143640*t);
+   saturn_y_3+=     0.00002981969 * cos( 4.42868951627 +      639.89728631400*t);
+   saturn_y_3+=     0.00002376221 * cos( 1.46232779180 +       14.22709400160*t);
+   saturn_y_3+=     0.00001024263 * cos( 4.72337917196 +      647.01083331480*t);
+   saturn_y_3=saturn_y_3 * t * t * t;
+
+   double saturn_y_4=0.0;
+   saturn_y_4+=     0.00007978886 * cos( 3.13229268011 +      206.18554843720*t);
+   saturn_y_4+=     0.00007868379 * cos( 4.84940260021 +      220.41264243880*t);
+   saturn_y_4+=     0.00001514835 * cos( 4.77675733867 +      227.52618943960*t);
+   saturn_y_4+=     0.00001225569 * cos( 3.84500138574 +      433.71173787680*t);
+   saturn_y_4=saturn_y_4 * t * t * t * t;
+
+   return saturn_y_0+saturn_y_1+saturn_y_2+saturn_y_3+saturn_y_4;
+}
+
+double vsop87a_milli::saturn_z(double t){
+   double saturn_z_0=0.0;
+   saturn_z_0+=     0.41356950940 * cos( 3.60234142982 +      213.29909543800*t);
+   saturn_z_0+=     0.01148283576 * cos( 2.85128367469 +      426.59819087600*t);
+   saturn_z_0+=     0.01214249867 * cos( 0.00000000000 +        0.00000000000*t);
+   saturn_z_0+=     0.00329280791 * cos( 0.57121407104 +      206.18554843720*t);
+   saturn_z_0+=     0.00286934048 * cos( 3.48073526693 +      220.41264243880*t);
+   saturn_z_0+=     0.00099076584 * cos( 4.73369511264 +        7.11354700080*t);
+   saturn_z_0+=     0.00057361820 * cos( 4.92611225093 +      110.20632121940*t);
+   saturn_z_0+=     0.00047738127 * cos( 2.10039779728 +      639.89728631400*t);
+   saturn_z_0+=     0.00043458803 * cos( 5.84904978051 +      419.48464387520*t);
+   saturn_z_0+=     0.00034565673 * cos( 5.42614229590 +      316.39186965660*t);
+   saturn_z_0+=     0.00016185391 * cos( 2.72987173675 +      433.71173787680*t);
+   saturn_z_0+=     0.00009001270 * cos( 1.38140102737 +      103.09277421860*t);
+   saturn_z_0+=     0.00011433574 * cos( 3.71662021072 +      529.69096509460*t);
+   saturn_z_0+=     0.00005398708 * cos( 5.13204892363 +      202.25339517410*t);
+   saturn_z_0+=     0.00003902467 * cos( 3.71499738796 +      323.50541665740*t);
+   saturn_z_0+=     0.00003709212 * cos( 5.05549348785 +      632.78373931320*t);
+   saturn_z_0+=     0.00003614100 * cos( 3.35210451276 +      227.52618943960*t);
+   saturn_z_0+=     0.00003379953 * cos( 2.13868919206 +       11.04570026390*t);
+   saturn_z_0+=     0.00003089874 * cos( 3.62572857085 +      209.36694217490*t);
+   saturn_z_0+=     0.00002683064 * cos( 4.87689555581 +      224.34479570190*t);
+   saturn_z_0+=     0.00002963493 * cos( 0.46490184985 +      217.23124870110*t);
+   saturn_z_0+=     0.00002343367 * cos( 1.34558278340 +      853.19638175200*t);
+   saturn_z_0+=     0.00002423663 * cos( 2.92907094760 +       63.73589830340*t);
+   saturn_z_0+=     0.00001701916 * cos( 1.89892525654 +      735.87651353180*t);
+   saturn_z_0+=     0.00001941205 * cos( 4.59421314662 +       14.22709400160*t);
+   saturn_z_0+=     0.00001990145 * cos( 0.73166053611 +      199.07200143640*t);
+   saturn_z_0+=     0.00001460265 * cos( 3.12851339724 +      522.57741809380*t);
+   saturn_z_0+=     0.00001148341 * cos( 4.41139213915 +      117.31986822020*t);
+   saturn_z_0+=     0.00001092809 * cos( 3.10679381209 +      216.48048917570*t);
+   saturn_z_0+=     0.00001015179 * cos( 1.97897195994 +      647.01083331480*t);
+   saturn_z_0+=     0.00001098254 * cos( 0.96097709156 +      210.11770170030*t);
+   saturn_z_0+=     0.00001028743 * cos( 2.11933059243 +      415.55249061210*t);
+
+   double saturn_z_1=0.0;
+   saturn_z_1+=     0.01906503283 * cos( 4.94544746116 +      213.29909543800*t);
+   saturn_z_1+=     0.00528301265 * cos( 3.14159265359 +        0.00000000000*t);
+   saturn_z_1+=     0.00130262284 * cos( 2.26140980879 +      206.18554843720*t);
+   saturn_z_1+=     0.00101466332 * cos( 1.79095829545 +      220.41264243880*t);
+   saturn_z_1+=     0.00085947578 * cos( 0.51612788497 +      426.59819087600*t);
+   saturn_z_1+=     0.00022257446 * cos( 3.07684015656 +        7.11354700080*t);
+   saturn_z_1+=     0.00016179946 * cos( 1.19987517506 +      419.48464387520*t);
+   saturn_z_1+=     0.00009117402 * cos( 6.17205626814 +      639.89728631400*t);
+   saturn_z_1+=     0.00007470703 * cos( 0.93135621171 +      433.71173787680*t);
+   saturn_z_1+=     0.00004966668 * cos( 0.19044864213 +      316.39186965660*t);
+   saturn_z_1+=     0.00003816564 * cos( 4.38284565245 +      110.20632121940*t);
+   saturn_z_1+=     0.00002724120 * cos( 1.65580138665 +      227.52618943960*t);
+   saturn_z_1+=     0.00001734540 * cos( 3.51628075636 +      103.09277421860*t);
+   saturn_z_1+=     0.00001541995 * cos( 2.42323572812 +      199.07200143640*t);
+   saturn_z_1+=     0.00001209302 * cos( 2.91140089093 +       14.22709400160*t);
+   saturn_z_1=saturn_z_1 * t;
+
+   double saturn_z_2=0.0;
+   saturn_z_2+=     0.00131275155 * cos( 0.08868998101 +      213.29909543800*t);
+   saturn_z_2+=     0.00030147649 * cos( 3.91396203887 +      206.18554843720*t);
+   saturn_z_2+=     0.00019322173 * cos( 0.09228748624 +      220.41264243880*t);
+   saturn_z_2+=     0.00006868926 * cos( 5.48420255395 +      426.59819087600*t);
+   saturn_z_2+=     0.00002826107 * cos( 1.36583318555 +        7.11354700080*t);
+   saturn_z_2+=     0.00002646332 * cos( 2.94607395955 +      419.48464387520*t);
+   saturn_z_2+=     0.00003138233 * cos( 0.00000000000 +        0.00000000000*t);
+   saturn_z_2+=     0.00001844798 * cos( 5.43612062856 +      433.71173787680*t);
+   saturn_z_2+=     0.00001055383 * cos( 6.23890785179 +      227.52618943960*t);
+   saturn_z_2+=     0.00001036435 * cos( 4.33916308552 +      639.89728631400*t);
+   saturn_z_2=saturn_z_2 * t * t;
+
+   double saturn_z_3=0.0;
+   saturn_z_3+=     0.00004559419 * cos( 1.70646871501 +      213.29909543800*t);
+   saturn_z_3+=     0.00004779074 * cos( 5.57723756330 +      206.18554843720*t);
+   saturn_z_3+=     0.00003965402 * cos( 0.00000000000 +        0.00000000000*t);
+   saturn_z_3+=     0.00002508242 * cos( 4.64959056313 +      220.41264243880*t);
+   saturn_z_3=saturn_z_3 * t * t * t;
+
+   return saturn_z_0+saturn_z_1+saturn_z_2+saturn_z_3;
+}
+
+double vsop87a_milli::uranus_x(double t){
+   double uranus_x_0=0.0;
+   uranus_x_0+=    19.17370730359 * cos( 5.48133416489 +       74.78159856730*t);
+   uranus_x_0+=     1.32272523872 * cos( 0.00000000000 +        0.00000000000*t);
+   uranus_x_0+=     0.44402496796 * cos( 1.65967519586 +      149.56319713460*t);
+   uranus_x_0+=     0.14668209481 * cos( 3.42395862804 +       73.29712585900*t);
+   uranus_x_0+=     0.14130269479 * cos( 4.39572927934 +       76.26607127560*t);
+   uranus_x_0+=     0.06201106178 * cos( 5.14043574125 +        1.48447270830*t);
+   uranus_x_0+=     0.01542951343 * cos( 4.12121838072 +      224.34479570190*t);
+   uranus_x_0+=     0.01444216660 * cos( 2.65117115201 +      148.07872442630*t);
+   uranus_x_0+=     0.00944995563 * cos( 1.65869338757 +       11.04570026390*t);
+   uranus_x_0+=     0.00657524815 * cos( 0.57595170636 +      151.04766984290*t);
+   uranus_x_0+=     0.00621624676 * cos( 3.05882246638 +       77.75054398390*t);
+   uranus_x_0+=     0.00585182542 * cos( 4.79934779678 +       71.81265315070*t);
+   uranus_x_0+=     0.00634000270 * cos( 4.09556589724 +       63.73589830340*t);
+   uranus_x_0+=     0.00547699056 * cos( 3.63127725056 +       85.82729883120*t);
+   uranus_x_0+=     0.00458219984 * cos( 3.90788284112 +        2.96894541660*t);
+   uranus_x_0+=     0.00496087649 * cos( 0.59947400861 +      529.69096509460*t);
+   uranus_x_0+=     0.00383625535 * cos( 6.18762010576 +      138.51749687070*t);
+   uranus_x_0+=     0.00267938156 * cos( 0.96885660137 +      213.29909543800*t);
+   uranus_x_0+=     0.00215368005 * cos( 5.30877641428 +       38.13303563780*t);
+   uranus_x_0+=     0.00145505389 * cos( 2.31759757085 +       70.84944530420*t);
+   uranus_x_0+=     0.00135340032 * cos( 5.51062460816 +       78.71375183040*t);
+   uranus_x_0+=     0.00119593859 * cos( 4.10138544267 +       39.61750834610*t);
+   uranus_x_0+=     0.00125105686 * cos( 2.51455273063 +      111.43016149680*t);
+   uranus_x_0+=     0.00111260244 * cos( 5.12252784325 +      222.86032299360*t);
+   uranus_x_0+=     0.00104619827 * cos( 3.90538916334 +      146.59425171800*t);
+   uranus_x_0+=     0.00110125387 * cos( 4.45473528724 +       35.16409022120*t);
+   uranus_x_0+=     0.00063584588 * cos( 0.29966233158 +      299.12639426920*t);
+   uranus_x_0+=     0.00053904041 * cos( 3.92590422507 +        3.93215326310*t);
+   uranus_x_0+=     0.00065066905 * cos( 3.73008452906 +      109.94568878850*t);
+   uranus_x_0+=     0.00039181662 * cos( 2.68841280769 +        4.45341812490*t);
+   uranus_x_0+=     0.00034341683 * cos( 3.03781661928 +      225.82926841020*t);
+   uranus_x_0+=     0.00033134636 * cos( 2.54201591218 +       65.22037101170*t);
+   uranus_x_0+=     0.00034555652 * cos( 1.84699329257 +       79.23501669220*t);
+   uranus_x_0+=     0.00033867050 * cos( 5.98418436103 +       70.32818044240*t);
+   uranus_x_0+=     0.00028371614 * cos( 2.58026657123 +      127.47179660680*t);
+   uranus_x_0+=     0.00035943348 * cos( 4.08754543016 +      202.25339517410*t);
+   uranus_x_0+=     0.00025208833 * cos( 5.30272144657 +        9.56122755560*t);
+   uranus_x_0+=     0.00023467802 * cos( 4.09729860322 +      145.63104387150*t);
+   uranus_x_0+=     0.00022963939 * cos( 5.51475073655 +       84.34282612290*t);
+   uranus_x_0+=     0.00031823951 * cos( 5.53948583244 +      152.53214255120*t);
+   uranus_x_0+=     0.00028384953 * cos( 6.01785430306 +      184.72728735580*t);
+   uranus_x_0+=     0.00026657176 * cos( 6.11027939727 +      160.60889739850*t);
+   uranus_x_0+=     0.00019676762 * cos( 5.53431398332 +       74.66972398270*t);
+   uranus_x_0+=     0.00019653873 * cos( 2.28660913421 +       74.89347315190*t);
+   uranus_x_0+=     0.00019954280 * cos( 0.57450958037 +       12.53017297220*t);
+   uranus_x_0+=     0.00018565067 * cos( 0.62225019017 +       52.69019803950*t);
+   uranus_x_0+=     0.00020084756 * cos( 4.47297488471 +       22.09140052780*t);
+   uranus_x_0+=     0.00019926329 * cos( 1.39878194708 +      112.91463420510*t);
+   uranus_x_0+=     0.00018575632 * cos( 5.70217475790 +       33.67961751290*t);
+   uranus_x_0+=     0.00016587870 * cos( 4.86920309163 +      108.46121608020*t);
+   uranus_x_0+=     0.00015171194 * cos( 2.88415453399 +       41.10198105440*t);
+   uranus_x_0+=     0.00011245800 * cos( 6.11597016146 +       71.60020482960*t);
+   uranus_x_0+=     0.00013948521 * cos( 6.27545694160 +      221.37585028530*t);
+   uranus_x_0+=     0.00010798350 * cos( 1.70031857078 +       77.96299230500*t);
+   uranus_x_0+=     0.00013593955 * cos( 2.55407820633 +       87.31177153950*t);
+   uranus_x_0+=     0.00011997848 * cos( 0.94875212305 +     1059.38193018920*t);
+   uranus_x_0+=     0.00012884351 * cos( 5.08737999470 +      145.10977900970*t);
+   uranus_x_0+=     0.00012394786 * cos( 6.21892878850 +       72.33391801250*t);
+   uranus_x_0+=     0.00012253318 * cos( 0.19452856525 +       36.64856292950*t);
+   uranus_x_0+=     0.00011538642 * cos( 1.77241794539 +       77.22927912210*t);
+   uranus_x_0+=     0.00008738409 * cos( 4.96956808452 +      186.21176006410*t);
+   uranus_x_0+=     0.00007095608 * cos( 1.30384750044 +      297.64192156090*t);
+   uranus_x_0+=     0.00006262602 * cos( 1.71385983783 +      153.49535039770*t);
+   uranus_x_0+=     0.00007487302 * cos( 0.11408470667 +      426.59819087600*t);
+   uranus_x_0+=     0.00007798974 * cos( 5.82410372587 +      340.77089204480*t);
+   uranus_x_0+=     0.00006669249 * cos( 5.08626589612 +       62.25142559510*t);
+   uranus_x_0+=     0.00005505358 * cos( 3.31282108025 +      140.00196957900*t);
+   uranus_x_0+=     0.00005372927 * cos( 4.12498282863 +       75.30286342910*t);
+   uranus_x_0+=     0.00005354242 * cos( 3.69263973447 +       74.26033370550*t);
+   uranus_x_0+=     0.00004478123 * cos( 1.11838191479 +       66.70484372000*t);
+   uranus_x_0+=     0.00004233075 * cos( 3.94913608184 +      265.98929347750*t);
+   uranus_x_0+=     0.00005038353 * cos( 4.68664376918 +       18.15924726470*t);
+   uranus_x_0+=     0.00004570470 * cos( 0.97536665751 +      183.24281464750*t);
+   uranus_x_0+=     0.00004751325 * cos( 4.95762395337 +       73.81839072080*t);
+   uranus_x_0+=     0.00004448651 * cos( 0.29436142982 +      114.39910691340*t);
+   uranus_x_0+=     0.00003312340 * cos( 0.52418923788 +       82.85835341460*t);
+   uranus_x_0+=     0.00004515952 * cos( 2.88576303120 +       75.74480641380*t);
+   uranus_x_0+=     0.00003559276 * cos( 1.47627607503 +        5.93789083320*t);
+   uranus_x_0+=     0.00003268117 * cos( 0.51827231333 +      220.41264243880*t);
+   uranus_x_0+=     0.00003578235 * cos( 1.11528903208 +      137.03302416240*t);
+   uranus_x_0+=     0.00003004737 * cos( 5.12122132051 +        7.11354700080*t);
+   uranus_x_0+=     0.00002882392 * cos( 2.76136583899 +      373.90799283650*t);
+   uranus_x_0+=     0.00002579454 * cos( 3.84784330333 +      277.03499374140*t);
+   uranus_x_0+=     0.00002597765 * cos( 0.22409539936 +       96.87299909510*t);
+   uranus_x_0+=     0.00002560744 * cos( 4.44236223450 +       80.19822453870*t);
+   uranus_x_0+=     0.00002722745 * cos( 6.09456175016 +      106.97674337190*t);
+   uranus_x_0+=     0.00002528025 * cos( 0.89508396542 +       68.84370773410*t);
+   uranus_x_0+=     0.00002631138 * cos( 0.04831552531 +      305.34616939270*t);
+   uranus_x_0+=     0.00002541716 * cos( 0.64495056482 +       32.19514480460*t);
+   uranus_x_0+=     0.00002241129 * cos( 5.22377697501 +        3.18139373770*t);
+   uranus_x_0+=     0.00001965145 * cos( 0.09207526632 +       20.60692781950*t);
+   uranus_x_0+=     0.00002232022 * cos( 0.63571664756 +       80.71948940050*t);
+   uranus_x_0+=     0.00001933814 * cos( 5.75490033864 +       74.73341445750*t);
+   uranus_x_0+=     0.00001933817 * cos( 2.06557585395 +       74.82978267710*t);
+   uranus_x_0+=     0.00002138391 * cos( 4.20897429922 +       74.52096613640*t);
+   uranus_x_0+=     0.00002126427 * cos( 3.61171465436 +       75.04223099820*t);
+   uranus_x_0+=     0.00002215516 * cos( 2.18613112875 +      259.50888592310*t);
+   uranus_x_0+=     0.00001891213 * cos( 5.49941424248 +      300.61086697750*t);
+   uranus_x_0+=     0.00001927679 * cos( 1.29228021932 +      159.12442469020*t);
+   uranus_x_0+=     0.00001796558 * cos( 5.73271543335 +       74.62153987290*t);
+   uranus_x_0+=     0.00001792522 * cos( 2.08789166984 +       74.94165726170*t);
+   uranus_x_0+=     0.00001873542 * cos( 4.23391867169 +      206.18554843720*t);
+   uranus_x_0+=     0.00002182901 * cos( 1.23755478345 +      479.28838891550*t);
+   uranus_x_0+=     0.00001860591 * cos( 1.67536711716 +       42.58645376270*t);
+   uranus_x_0+=     0.00002075591 * cos( 3.15586933464 +      131.40394986990*t);
+   uranus_x_0+=     0.00001912582 * cos( 5.83091918696 +       14.97785352700*t);
+   uranus_x_0+=     0.00001529180 * cos( 2.05204104820 +      191.20769491020*t);
+   uranus_x_0+=     0.00002064173 * cos( 3.60208606410 +      835.03713448730*t);
+   uranus_x_0+=     0.00001768763 * cos( 1.19254481620 +      219.89137757700*t);
+   uranus_x_0+=     0.00001892359 * cos( 4.32128621847 +      154.01661525950*t);
+   uranus_x_0+=     0.00001797047 * cos( 1.73417465594 +      227.31374111850*t);
+   uranus_x_0+=     0.00001558489 * cos( 6.16891070489 +       59.80374504030*t);
+   uranus_x_0+=     0.00001672893 * cos( 0.01232646186 +      143.62530630140*t);
+   uranus_x_0+=     0.00001467268 * cos( 2.10975578758 +        2.44768055480*t);
+   uranus_x_0+=     0.00001677659 * cos( 0.42525121334 +        8.07675484730*t);
+   uranus_x_0+=     0.00001347303 * cos( 5.46763140224 +      288.08069400530*t);
+   uranus_x_0+=     0.00001744555 * cos( 0.82022450313 +       56.62235130260*t);
+   uranus_x_0+=     0.00001427180 * cos( 0.38786175669 +       92.94084583200*t);
+   uranus_x_0+=     0.00001263032 * cos( 5.63689596853 +      404.50679034820*t);
+   uranus_x_0+=     0.00001221506 * cos( 5.20012455894 +       54.17467074780*t);
+   uranus_x_0+=     0.00001574905 * cos( 5.72297800263 +       39.35687591520*t);
+   uranus_x_0+=     0.00001269686 * cos( 2.66330104031 +      142.44965013380*t);
+   uranus_x_0+=     0.00001438869 * cos( 0.72633739717 +      522.57741809380*t);
+   uranus_x_0+=     0.00001408997 * cos( 3.61751904356 +      536.80451209540*t);
+   uranus_x_0+=     0.00001418300 * cos( 2.29718712012 +      235.39049596580*t);
+   uranus_x_0+=     0.00001040948 * cos( 2.74644165501 +        5.41662597140*t);
+   uranus_x_0+=     0.00001256867 * cos( 5.61684736425 +       67.66805156650*t);
+   uranus_x_0+=     0.00001164218 * cos( 2.08302637541 +       81.89514556810*t);
+   uranus_x_0+=     0.00001009252 * cos( 2.02320517037 +       74.03083904190*t);
+   uranus_x_0+=     0.00001077810 * cos( 1.05685900920 +      128.95626931510*t);
+   uranus_x_0+=     0.00001212025 * cos( 3.41577832660 +      211.81462272970*t);
+   uranus_x_0+=     0.00001223648 * cos( 3.84373514640 +      187.69623277240*t);
+   uranus_x_0+=     0.00001334371 * cos( 2.17621743689 +      380.12776796000*t);
+   uranus_x_0+=     0.00001144565 * cos( 2.42148845239 +      296.15744885260*t);
+   uranus_x_0+=     0.00001241589 * cos( 1.81282962357 +      134.58534360760*t);
+   uranus_x_0+=     0.00001192274 * cos( 5.58661990233 +       50.40257617910*t);
+   uranus_x_0+=     0.00001093678 * cos( 3.94451812233 +      230.56457082540*t);
+   uranus_x_0+=     0.00001077916 * cos( 5.01417740021 +      181.75834193920*t);
+   uranus_x_0+=     0.00001166898 * cos( 2.02955848543 +      110.20632121940*t);
+   uranus_x_0+=     0.00001056848 * cos( 5.53440854164 +       14.01464568050*t);
+
+   double uranus_x_1=0.0;
+   uranus_x_1+=     0.00739730021 * cos( 6.01067825116 +      149.56319713460*t);
+   uranus_x_1+=     0.00526878306 * cos( 3.14159265359 +        0.00000000000*t);
+   uranus_x_1+=     0.00239840801 * cos( 5.33657762707 +       73.29712585900*t);
+   uranus_x_1+=     0.00229676787 * cos( 2.48204455775 +       76.26607127560*t);
+   uranus_x_1+=     0.00111045158 * cos( 5.57157235960 +       11.04570026390*t);
+   uranus_x_1+=     0.00096352822 * cos( 0.35070389084 +       63.73589830340*t);
+   uranus_x_1+=     0.00081511870 * cos( 1.21058618039 +       85.82729883120*t);
+   uranus_x_1+=     0.00045687564 * cos( 2.29216583843 +      138.51749687070*t);
+   uranus_x_1+=     0.00051382501 * cos( 2.18935125260 +      224.34479570190*t);
+   uranus_x_1+=     0.00038844330 * cos( 0.30724575951 +       70.84944530420*t);
+   uranus_x_1+=     0.00036158493 * cos( 1.23634798757 +       78.71375183040*t);
+   uranus_x_1+=     0.00032333094 * cos( 5.06666556704 +       74.78159856730*t);
+   uranus_x_1+=     0.00021685656 * cos( 4.93710968392 +      151.04766984290*t);
+   uranus_x_1+=     0.00019441970 * cos( 1.30617490304 +       77.75054398390*t);
+   uranus_x_1+=     0.00017376241 * cos( 0.24607221230 +       71.81265315070*t);
+   uranus_x_1+=     0.00015211071 * cos( 5.53141633140 +        3.93215326310*t);
+   uranus_x_1+=     0.00007735984 * cos( 1.61349552789 +       71.60020482960*t);
+   uranus_x_1+=     0.00007425078 * cos( 6.20357977116 +       77.96299230500*t);
+   uranus_x_1+=     0.00006995857 * cos( 2.40633283814 +      145.63104387150*t);
+   uranus_x_1+=     0.00007291691 * cos( 2.23597571444 +        2.96894541660*t);
+   uranus_x_1+=     0.00007585264 * cos( 2.76074218330 +      148.07872442630*t);
+   uranus_x_1+=     0.00004378335 * cos( 3.74296322240 +      160.60889739850*t);
+   uranus_x_1+=     0.00004127713 * cos( 1.48475181305 +       22.09140052780*t);
+   uranus_x_1+=     0.00003933541 * cos( 4.73864204208 +       65.22037101170*t);
+   uranus_x_1+=     0.00002910312 * cos( 5.91941333050 +      127.47179660680*t);
+   uranus_x_1+=     0.00002788434 * cos( 4.90117297196 +      213.29909543800*t);
+   uranus_x_1+=     0.00002802392 * cos( 3.76505436434 +       52.69019803950*t);
+   uranus_x_1+=     0.00002545264 * cos( 3.36768337628 +        9.56122755560*t);
+   uranus_x_1+=     0.00002647073 * cos( 4.53813176345 +       12.53017297220*t);
+   uranus_x_1+=     0.00003177614 * cos( 4.65226634926 +      299.12639426920*t);
+   uranus_x_1+=     0.00002228396 * cos( 0.18087986338 +       87.31177153950*t);
+   uranus_x_1+=     0.00002824668 * cos( 3.40143685673 +       84.34282612290*t);
+   uranus_x_1+=     0.00002525203 * cos( 2.83821144961 +       18.15924726470*t);
+   uranus_x_1+=     0.00002216987 * cos( 4.78338909951 +       72.33391801250*t);
+   uranus_x_1+=     0.00001624493 * cos( 3.75817281127 +      153.49535039770*t);
+   uranus_x_1+=     0.00001928881 * cos( 2.39940180311 +       39.61750834610*t);
+   uranus_x_1+=     0.00001555444 * cos( 4.13741667297 +       73.81839072080*t);
+   uranus_x_1+=     0.00001600865 * cos( 0.08376247543 +       79.23501669220*t);
+   uranus_x_1+=     0.00001476317 * cos( 3.67283851029 +       75.74480641380*t);
+   uranus_x_1+=     0.00001427088 * cos( 1.45690759014 +       70.32818044240*t);
+   uranus_x_1+=     0.00001533469 * cos( 3.71776498048 +      152.53214255120*t);
+   uranus_x_1+=     0.00001747266 * cos( 3.24870046809 +       77.22927912210*t);
+   uranus_x_1+=     0.00001716831 * cos( 3.39415662657 +      222.86032299360*t);
+   uranus_x_1+=     0.00001707504 * cos( 1.11296012106 +      225.82926841020*t);
+   uranus_x_1+=     0.00001734228 * cos( 5.39619902298 +      146.59425171800*t);
+   uranus_x_1+=     0.00001476617 * cos( 3.53047075439 +        3.18139373770*t);
+   uranus_x_1+=     0.00001179645 * cos( 5.13953276367 +      220.41264243880*t);
+   uranus_x_1+=     0.00001239315 * cos( 0.98221206501 +        4.45341812490*t);
+   uranus_x_1+=     0.00001099691 * cos( 1.35138854505 +       62.25142559510*t);
+   uranus_x_1+=     0.00001061312 * cos( 5.00125105380 +      131.40394986990*t);
+   uranus_x_1+=     0.00001029261 * cos( 5.42434597865 +      109.94568878850*t);
+   uranus_x_1=uranus_x_1 * t;
+
+   double uranus_x_2=0.0;
+   uranus_x_2+=     0.00016015732 * cos( 3.83700026619 +       74.78159856730*t);
+   uranus_x_2+=     0.00010915299 * cos( 3.02987776270 +      149.56319713460*t);
+   uranus_x_2+=     0.00007497619 * cos( 3.83429136661 +       11.04570026390*t);
+   uranus_x_2+=     0.00008053623 * cos( 2.54646146122 +       63.73589830340*t);
+   uranus_x_2+=     0.00005408033 * cos( 4.78033642303 +       70.84944530420*t);
+   uranus_x_2+=     0.00005021971 * cos( 3.04632772928 +       78.71375183040*t);
+   uranus_x_2+=     0.00006717313 * cos( 5.31264214501 +       85.82729883120*t);
+   uranus_x_2+=     0.00005284684 * cos( 2.11901942097 +       73.29712585900*t);
+   uranus_x_2+=     0.00004874936 * cos( 5.68616132176 +       76.26607127560*t);
+   uranus_x_2+=     0.00003002124 * cos( 4.07944398452 +      138.51749687070*t);
+   uranus_x_2+=     0.00002521797 * cos( 3.36028253173 +       71.60020482960*t);
+   uranus_x_2+=     0.00002413832 * cos( 4.45865225690 +       77.96299230500*t);
+   uranus_x_2+=     0.00002221373 * cos( 0.87427485235 +        3.93215326310*t);
+   uranus_x_2+=     0.00002291767 * cos( 0.00000000000 +        0.00000000000*t);
+   uranus_x_2+=     0.00001040250 * cos( 0.73133408837 +      145.63104387150*t);
+   uranus_x_2+=     0.00001046100 * cos( 5.85311910228 +      224.34479570190*t);
+   uranus_x_2=uranus_x_2 * t * t;
+
+   double uranus_x_3=0.0;
+   uranus_x_3+=     0.00001307049 * cos( 0.00000000000 +        0.00000000000*t);
+   uranus_x_3=uranus_x_3 * t * t * t;
+
+   return uranus_x_0+uranus_x_1+uranus_x_2+uranus_x_3;
+}
+
+double vsop87a_milli::uranus_y(double t){
+   double uranus_y_0=0.0;
+   uranus_y_0+=    19.16518231584 * cos( 3.91045677002 +       74.78159856730*t);
+   uranus_y_0+=     0.44390465203 * cos( 0.08884111329 +      149.56319713460*t);
+   uranus_y_0+=     0.16256125476 * cos( 3.14159265359 +        0.00000000000*t);
+   uranus_y_0+=     0.14755940186 * cos( 1.85423280679 +       73.29712585900*t);
+   uranus_y_0+=     0.14123958128 * cos( 2.82486076549 +       76.26607127560*t);
+   uranus_y_0+=     0.06250078231 * cos( 3.56960243857 +        1.48447270830*t);
+   uranus_y_0+=     0.01542668264 * cos( 2.55040539213 +      224.34479570190*t);
+   uranus_y_0+=     0.01442356575 * cos( 1.08004542712 +      148.07872442630*t);
+   uranus_y_0+=     0.00938975501 * cos( 0.09275714761 +       11.04570026390*t);
+   uranus_y_0+=     0.00650331846 * cos( 2.76142680222 +       63.73589830340*t);
+   uranus_y_0+=     0.00657343120 * cos( 5.28830704469 +      151.04766984290*t);
+   uranus_y_0+=     0.00621326770 * cos( 1.48795811387 +       77.75054398390*t);
+   uranus_y_0+=     0.00541961958 * cos( 3.24476486661 +       71.81265315070*t);
+   uranus_y_0+=     0.00547472694 * cos( 2.06037924573 +       85.82729883120*t);
+   uranus_y_0+=     0.00459589120 * cos( 2.33745536070 +        2.96894541660*t);
+   uranus_y_0+=     0.00495936105 * cos( 5.31205753740 +      529.69096509460*t);
+   uranus_y_0+=     0.00387922853 * cos( 4.62026923885 +      138.51749687070*t);
+   uranus_y_0+=     0.00268363417 * cos( 5.68085299020 +      213.29909543800*t);
+   uranus_y_0+=     0.00216239629 * cos( 3.73800767580 +       38.13303563780*t);
+   uranus_y_0+=     0.00144032475 * cos( 0.75015700920 +       70.84944530420*t);
+   uranus_y_0+=     0.00135290820 * cos( 3.93970260616 +       78.71375183040*t);
+   uranus_y_0+=     0.00119670613 * cos( 2.53058783780 +       39.61750834610*t);
+   uranus_y_0+=     0.00124868545 * cos( 0.94315917319 +      111.43016149680*t);
+   uranus_y_0+=     0.00111204860 * cos( 3.55163219419 +      222.86032299360*t);
+   uranus_y_0+=     0.00104507929 * cos( 2.33345675603 +      146.59425171800*t);
+   uranus_y_0+=     0.00108584454 * cos( 6.02234848388 +       35.16409022120*t);
+   uranus_y_0+=     0.00063573747 * cos( 5.01204967920 +      299.12639426920*t);
+   uranus_y_0+=     0.00053289771 * cos( 2.38437587876 +        3.93215326310*t);
+   uranus_y_0+=     0.00063774261 * cos( 2.15607602904 +      109.94568878850*t);
+   uranus_y_0+=     0.00039218598 * cos( 1.11841109252 +        4.45341812490*t);
+   uranus_y_0+=     0.00034205426 * cos( 0.92405922576 +       65.22037101170*t);
+   uranus_y_0+=     0.00034334377 * cos( 1.46696169843 +      225.82926841020*t);
+   uranus_y_0+=     0.00034538316 * cos( 0.27613780697 +       79.23501669220*t);
+   uranus_y_0+=     0.00039256771 * cos( 5.75956853703 +      202.25339517410*t);
+   uranus_y_0+=     0.00026157754 * cos( 3.74097610798 +        9.56122755560*t);
+   uranus_y_0+=     0.00023427328 * cos( 2.52740125551 +      145.63104387150*t);
+   uranus_y_0+=     0.00022933138 * cos( 3.94455540350 +       84.34282612290*t);
+   uranus_y_0+=     0.00031816303 * cos( 3.96860170484 +      152.53214255120*t);
+   uranus_y_0+=     0.00025237176 * cos( 4.45141413666 +       70.32818044240*t);
+   uranus_y_0+=     0.00028372491 * cos( 4.44714627097 +      184.72728735580*t);
+   uranus_y_0+=     0.00026652859 * cos( 4.53944395347 +      160.60889739850*t);
+   uranus_y_0+=     0.00019666208 * cos( 3.96350065335 +       74.66972398270*t);
+   uranus_y_0+=     0.00019643845 * cos( 0.71577796385 +       74.89347315190*t);
+   uranus_y_0+=     0.00019838981 * cos( 5.29113397354 +       12.53017297220*t);
+   uranus_y_0+=     0.00021523908 * cos( 4.93565132068 +       36.64856292950*t);
+   uranus_y_0+=     0.00015537967 * cos( 1.87863275460 +       52.69019803950*t);
+   uranus_y_0+=     0.00020115100 * cos( 3.45473780762 +      127.47179660680*t);
+   uranus_y_0+=     0.00020051641 * cos( 2.90386352937 +       22.09140052780*t);
+   uranus_y_0+=     0.00019901477 * cos( 6.11075402434 +      112.91463420510*t);
+   uranus_y_0+=     0.00018126776 * cos( 0.98478853787 +       33.67961751290*t);
+   uranus_y_0+=     0.00015174962 * cos( 1.31314034959 +       41.10198105440*t);
+   uranus_y_0+=     0.00011239020 * cos( 4.54508334011 +       71.60020482960*t);
+   uranus_y_0+=     0.00013948849 * cos( 4.70474945682 +      221.37585028530*t);
+   uranus_y_0+=     0.00010819728 * cos( 0.12807029856 +       77.96299230500*t);
+   uranus_y_0+=     0.00013589665 * cos( 0.98313719930 +       87.31177153950*t);
+   uranus_y_0+=     0.00011996772 * cos( 5.66129275335 +     1059.38193018920*t);
+   uranus_y_0+=     0.00012407787 * cos( 4.64945783340 +       72.33391801250*t);
+   uranus_y_0+=     0.00011531140 * cos( 0.20190074645 +       77.22927912210*t);
+   uranus_y_0+=     0.00008736150 * cos( 3.39874828293 +      186.21176006410*t);
+   uranus_y_0+=     0.00007093587 * cos( 6.01613487245 +      297.64192156090*t);
+   uranus_y_0+=     0.00006408245 * cos( 3.93246367895 +       62.25142559510*t);
+   uranus_y_0+=     0.00006261153 * cos( 0.14258542752 +      153.49535039770*t);
+   uranus_y_0+=     0.00007494000 * cos( 4.82565771386 +      426.59819087600*t);
+   uranus_y_0+=     0.00007856014 * cos( 1.12354254831 +      340.77089204480*t);
+   uranus_y_0+=     0.00005516018 * cos( 1.73758326119 +      140.00196957900*t);
+   uranus_y_0+=     0.00005556643 * cos( 3.68095215063 +      145.10977900970*t);
+   uranus_y_0+=     0.00005368405 * cos( 2.55422957958 +       75.30286342910*t);
+   uranus_y_0+=     0.00005350948 * cos( 2.12171493922 +       74.26033370550*t);
+   uranus_y_0+=     0.00004508794 * cos( 5.82224064821 +       66.70484372000*t);
+   uranus_y_0+=     0.00004290374 * cos( 5.54490766551 +      265.98929347750*t);
+   uranus_y_0+=     0.00005013871 * cos( 3.11907749268 +       18.15924726470*t);
+   uranus_y_0+=     0.00004326138 * cos( 5.70135056853 +      183.24281464750*t);
+   uranus_y_0+=     0.00004750018 * cos( 3.38678300054 +       73.81839072080*t);
+   uranus_y_0+=     0.00004445347 * cos( 5.00638490308 +      114.39910691340*t);
+   uranus_y_0+=     0.00003314154 * cos( 5.23054574329 +       82.85835341460*t);
+   uranus_y_0+=     0.00004509054 * cos( 1.31254342829 +       75.74480641380*t);
+   uranus_y_0+=     0.00003553107 * cos( 6.18906516846 +        5.93789083320*t);
+   uranus_y_0+=     0.00003265634 * cos( 5.23063560176 +      220.41264243880*t);
+   uranus_y_0+=     0.00003575435 * cos( 5.83994849224 +      137.03302416240*t);
+   uranus_y_0+=     0.00002880734 * cos( 1.19038424330 +      373.90799283650*t);
+   uranus_y_0+=     0.00002885443 * cos( 3.50279993038 +        7.11354700080*t);
+   uranus_y_0+=     0.00002594155 * cos( 4.93691413537 +       96.87299909510*t);
+   uranus_y_0+=     0.00002559357 * cos( 2.87184237678 +       80.19822453870*t);
+   uranus_y_0+=     0.00002676616 * cos( 1.61805362044 +      305.34616939270*t);
+   uranus_y_0+=     0.00002246530 * cos( 5.80081898763 +      108.46121608020*t);
+   uranus_y_0+=     0.00002474483 * cos( 2.21173751117 +       32.19514480460*t);
+   uranus_y_0+=     0.00002060991 * cos( 6.24178596384 +       56.62235130260*t);
+   uranus_y_0+=     0.00001958711 * cos( 4.80807045815 +       20.60692781950*t);
+   uranus_y_0+=     0.00002227451 * cos( 5.34765264557 +       80.71948940050*t);
+   uranus_y_0+=     0.00001937874 * cos( 0.49529839431 +       74.82978267710*t);
+   uranus_y_0+=     0.00001937871 * cos( 4.18462288684 +       74.73341445750*t);
+   uranus_y_0+=     0.00002164382 * cos( 0.47581392325 +        3.18139373770*t);
+   uranus_y_0+=     0.00002138407 * cos( 2.63817804331 +       74.52096613640*t);
+   uranus_y_0+=     0.00002130909 * cos( 2.04143912495 +       75.04223099820*t);
+   uranus_y_0+=     0.00001787737 * cos( 0.32096699926 +        2.44768055480*t);
+   uranus_y_0+=     0.00002212861 * cos( 0.61491281306 +      259.50888592310*t);
+   uranus_y_0+=     0.00001889369 * cos( 3.92852240171 +      300.61086697750*t);
+   uranus_y_0+=     0.00002275258 * cos( 1.55666401505 +      131.40394986990*t);
+   uranus_y_0+=     0.00001925946 * cos( 6.00527473515 +      159.12442469020*t);
+   uranus_y_0+=     0.00001802494 * cos( 4.16218259902 +       74.62153987290*t);
+   uranus_y_0+=     0.00001796292 * cos( 0.51761494342 +       74.94165726170*t);
+   uranus_y_0+=     0.00002240648 * cos( 0.47739127862 +      181.75834193920*t);
+   uranus_y_0+=     0.00001924499 * cos( 2.64284880495 +      206.18554843720*t);
+   uranus_y_0+=     0.00001626134 * cos( 3.70023731184 +      191.20769491020*t);
+   uranus_y_0+=     0.00001860824 * cos( 0.10445996392 +       42.58645376270*t);
+   uranus_y_0+=     0.00002177437 * cos( 2.80437422101 +      479.28838891550*t);
+   uranus_y_0+=     0.00001896184 * cos( 4.26975898003 +       14.97785352700*t);
+   uranus_y_0+=     0.00002045249 * cos( 5.17400788104 +      835.03713448730*t);
+   uranus_y_0+=     0.00001887812 * cos( 2.75000237791 +      154.01661525950*t);
+   uranus_y_0+=     0.00001794754 * cos( 0.16290844853 +      227.31374111850*t);
+   uranus_y_0+=     0.00001347410 * cos( 3.89237011696 +      288.08069400530*t);
+   uranus_y_0+=     0.00001572826 * cos( 5.93367812903 +      219.89137757700*t);
+   uranus_y_0+=     0.00001424804 * cos( 5.10004758033 +       92.94084583200*t);
+   uranus_y_0+=     0.00001267766 * cos( 0.92771324396 +      404.50679034820*t);
+   uranus_y_0+=     0.00001588897 * cos( 4.15115668974 +       39.35687591520*t);
+   uranus_y_0+=     0.00001269786 * cos( 1.09685727529 +      142.44965013380*t);
+   uranus_y_0+=     0.00001291065 * cos( 5.67425699047 +       68.84370773410*t);
+   uranus_y_0+=     0.00001436850 * cos( 5.44312198350 +      522.57741809380*t);
+   uranus_y_0+=     0.00001405564 * cos( 2.04677392527 +      536.80451209540*t);
+   uranus_y_0+=     0.00001416917 * cos( 0.72597245494 +      235.39049596580*t);
+   uranus_y_0+=     0.00001165315 * cos( 0.51071041452 +       81.89514556810*t);
+   uranus_y_0+=     0.00001035262 * cos( 1.20639876458 +        5.41662597140*t);
+   uranus_y_0+=     0.00001009454 * cos( 0.45375065997 +       74.03083904190*t);
+   uranus_y_0+=     0.00001220696 * cos( 1.84988185963 +      211.81462272970*t);
+   uranus_y_0+=     0.00001222886 * cos( 2.27306099902 +      187.69623277240*t);
+   uranus_y_0+=     0.00001336792 * cos( 3.74888989756 +      380.12776796000*t);
+   uranus_y_0+=     0.00001151803 * cos( 0.46579056125 +      128.95626931510*t);
+   uranus_y_0+=     0.00001149114 * cos( 0.85101218281 +      296.15744885260*t);
+   uranus_y_0+=     0.00001163762 * cos( 5.51157783762 +      230.56457082540*t);
+   uranus_y_0+=     0.00001192292 * cos( 0.68084398426 +       99.16062095550*t);
+   uranus_y_0+=     0.00001151286 * cos( 4.01147735438 +       67.66805156650*t);
+   uranus_y_0+=     0.00001189801 * cos( 4.01778306134 +       50.40257617910*t);
+   uranus_y_0+=     0.00001015998 * cos( 1.00290501307 +       35.42472265210*t);
+   uranus_y_0+=     0.00001174953 * cos( 0.45683512473 +      110.20632121940*t);
+   uranus_y_0+=     0.00001051606 * cos( 3.96907647535 +       14.01464568050*t);
+
+   double uranus_y_1=0.0;
+   uranus_y_1+=     0.02157896385 * cos( 0.00000000000 +        0.00000000000*t);
+   uranus_y_1+=     0.00739227349 * cos( 4.43963890935 +      149.56319713460*t);
+   uranus_y_1+=     0.00238545685 * cos( 3.76882493145 +       73.29712585900*t);
+   uranus_y_1+=     0.00229396424 * cos( 0.91090183978 +       76.26607127560*t);
+   uranus_y_1+=     0.00110137111 * cos( 4.00844441616 +       11.04570026390*t);
+   uranus_y_1+=     0.00094979054 * cos( 5.07141537066 +       63.73589830340*t);
+   uranus_y_1+=     0.00081474163 * cos( 5.92275367106 +       85.82729883120*t);
+   uranus_y_1+=     0.00045457174 * cos( 0.73292241207 +      138.51749687070*t);
+   uranus_y_1+=     0.00051366974 * cos( 0.61844114994 +      224.34479570190*t);
+   uranus_y_1+=     0.00038296005 * cos( 5.01873578671 +       70.84944530420*t);
+   uranus_y_1+=     0.00036146116 * cos( 5.94859452787 +       78.71375183040*t);
+   uranus_y_1+=     0.00032420558 * cos( 4.32617271732 +       74.78159856730*t);
+   uranus_y_1+=     0.00021673269 * cos( 3.36607263522 +      151.04766984290*t);
+   uranus_y_1+=     0.00019425087 * cos( 6.01842187783 +       77.75054398390*t);
+   uranus_y_1+=     0.00017393206 * cos( 4.96098895488 +       71.81265315070*t);
+   uranus_y_1+=     0.00014991169 * cos( 3.97176856758 +        3.93215326310*t);
+   uranus_y_1+=     0.00007732367 * cos( 0.04256630122 +       71.60020482960*t);
+   uranus_y_1+=     0.00007438492 * cos( 4.63165436478 +       77.96299230500*t);
+   uranus_y_1+=     0.00006979238 * cos( 0.83723520791 +      145.63104387150*t);
+   uranus_y_1+=     0.00007321559 * cos( 0.66348425538 +        2.96894541660*t);
+   uranus_y_1+=     0.00007595636 * cos( 1.19807643487 +      148.07872442630*t);
+   uranus_y_1+=     0.00004376824 * cos( 2.17182724016 +      160.60889739850*t);
+   uranus_y_1+=     0.00003962568 * cos( 3.18042711824 +       65.22037101170*t);
+   uranus_y_1+=     0.00004117202 * cos( 6.19931612790 +       22.09140052780*t);
+   uranus_y_1+=     0.00002830313 * cos( 3.34365222278 +      213.29909543800*t);
+   uranus_y_1+=     0.00002690065 * cos( 1.78946471198 +        9.56122755560*t);
+   uranus_y_1+=     0.00002628159 * cos( 2.97459067399 +       12.53017297220*t);
+   uranus_y_1+=     0.00003174617 * cos( 3.08131638838 +      299.12639426920*t);
+   uranus_y_1+=     0.00002227510 * cos( 4.89407437055 +       87.31177153950*t);
+   uranus_y_1+=     0.00002821206 * cos( 1.83130010947 +       84.34282612290*t);
+   uranus_y_1+=     0.00002510693 * cos( 1.27166561854 +       18.15924726470*t);
+   uranus_y_1+=     0.00002224655 * cos( 3.21541108798 +       72.33391801250*t);
+   uranus_y_1+=     0.00001929662 * cos( 4.68383962079 +       52.69019803950*t);
+   uranus_y_1+=     0.00001707606 * cos( 5.60415260609 +      127.47179660680*t);
+   uranus_y_1+=     0.00001620819 * cos( 2.18676505386 +      153.49535039770*t);
+   uranus_y_1+=     0.00001926926 * cos( 0.82821252139 +       39.61750834610*t);
+   uranus_y_1+=     0.00001555611 * cos( 2.56681954823 +       73.81839072080*t);
+   uranus_y_1+=     0.00001595703 * cos( 4.79558057860 +       79.23501669220*t);
+   uranus_y_1+=     0.00001476550 * cos( 2.10111459539 +       75.74480641380*t);
+   uranus_y_1+=     0.00001531355 * cos( 2.14701519407 +      152.53214255120*t);
+   uranus_y_1+=     0.00001744564 * cos( 1.67897185084 +       77.22927912210*t);
+   uranus_y_1+=     0.00001713945 * cos( 1.82334975258 +      222.86032299360*t);
+   uranus_y_1+=     0.00001705184 * cos( 5.82532917611 +      225.82926841020*t);
+   uranus_y_1+=     0.00001739019 * cos( 3.82452086292 +      146.59425171800*t);
+   uranus_y_1+=     0.00001291753 * cos( 6.19666243545 +       70.32818044240*t);
+   uranus_y_1+=     0.00001458526 * cos( 5.10147126404 +        3.18139373770*t);
+   uranus_y_1+=     0.00001179966 * cos( 3.56807126055 +      220.41264243880*t);
+   uranus_y_1+=     0.00001234914 * cos( 5.69239889831 +        4.45341812490*t);
+   uranus_y_1+=     0.00001075108 * cos( 1.76286452034 +       56.62235130260*t);
+   uranus_y_1+=     0.00001035661 * cos( 6.12642568708 +       62.25142559510*t);
+   uranus_y_1+=     0.00001114523 * cos( 3.41304662369 +      131.40394986990*t);
+   uranus_y_1+=     0.00001050191 * cos( 3.84176879347 +      109.94568878850*t);
+   uranus_y_1=uranus_y_1 * t;
+
+   double uranus_y_2=0.0;
+   uranus_y_2+=     0.00034812647 * cos( 3.14159265359 +        0.00000000000*t);
+   uranus_y_2+=     0.00016589194 * cos( 2.29556740620 +       74.78159856730*t);
+   uranus_y_2+=     0.00010905147 * cos( 1.45737963668 +      149.56319713460*t);
+   uranus_y_2+=     0.00007484633 * cos( 2.27968076918 +       11.04570026390*t);
+   uranus_y_2+=     0.00007964298 * cos( 0.97230247087 +       63.73589830340*t);
+   uranus_y_2+=     0.00005307100 * cos( 3.20519221878 +       70.84944530420*t);
+   uranus_y_2+=     0.00005018595 * cos( 1.47518527303 +       78.71375183040*t);
+   uranus_y_2+=     0.00006713255 * cos( 3.74148881189 +       85.82729883120*t);
+   uranus_y_2+=     0.00005265170 * cos( 0.54901216905 +       73.29712585900*t);
+   uranus_y_2+=     0.00004864822 * cos( 4.11367426823 +       76.26607127560*t);
+   uranus_y_2+=     0.00002995853 * cos( 2.49432193549 +      138.51749687070*t);
+   uranus_y_2+=     0.00002519021 * cos( 1.78896824345 +       71.60020482960*t);
+   uranus_y_2+=     0.00002418371 * cos( 2.88675006488 +       77.96299230500*t);
+   uranus_y_2+=     0.00002185856 * cos( 5.58862614977 +        3.93215326310*t);
+   uranus_y_2+=     0.00001035578 * cos( 5.44752448275 +      145.63104387150*t);
+   uranus_y_2+=     0.00001044459 * cos( 4.27972239360 +      224.34479570190*t);
+   uranus_y_2=uranus_y_2 * t * t;
+
+   double uranus_y_3=0.0;
+   uranus_y_3+=     0.00001211380 * cos( 0.00000000000 +        0.00000000000*t);
+   uranus_y_3=uranus_y_3 * t * t * t;
+
+   return uranus_y_0+uranus_y_1+uranus_y_2+uranus_y_3;
+}
+
+double vsop87a_milli::uranus_z(double t){
+   double uranus_z_0=0.0;
+   uranus_z_0+=     0.25878127698 * cos( 2.61861272578 +       74.78159856730*t);
+   uranus_z_0+=     0.01774318778 * cos( 3.14159265359 +        0.00000000000*t);
+   uranus_z_0+=     0.00599316131 * cos( 5.08119500585 +      149.56319713460*t);
+   uranus_z_0+=     0.00190281890 * cos( 1.61643841193 +       76.26607127560*t);
+   uranus_z_0+=     0.00190881685 * cos( 0.57869575952 +       73.29712585900*t);
+   uranus_z_0+=     0.00084626761 * cos( 2.26030150166 +        1.48447270830*t);
+   uranus_z_0+=     0.00030734257 * cos( 0.23571721555 +       63.73589830340*t);
+   uranus_z_0+=     0.00020842052 * cos( 1.26054208091 +      224.34479570190*t);
+   uranus_z_0+=     0.00019734273 * cos( 6.04314677688 +      148.07872442630*t);
+   uranus_z_0+=     0.00012537530 * cos( 5.17169051466 +       11.04570026390*t);
+   uranus_z_0+=     0.00014582864 * cos( 6.14852037212 +       71.81265315070*t);
+   uranus_z_0+=     0.00010407529 * cos( 3.65320417038 +      213.29909543800*t);
+   uranus_z_0+=     0.00011261541 * cos( 3.55973769686 +      529.69096509460*t);
+   uranus_z_0+=     0.00008855669 * cos( 4.03774505739 +      151.04766984290*t);
+   uranus_z_0+=     0.00008239460 * cos( 0.34225652715 +       77.75054398390*t);
+   uranus_z_0+=     0.00007950169 * cos( 0.72564903051 +       85.82729883120*t);
+   uranus_z_0+=     0.00006867469 * cos( 0.81417174224 +        2.96894541660*t);
+   uranus_z_0+=     0.00005648720 * cos( 3.45324719543 +      138.51749687070*t);
+   uranus_z_0+=     0.00004581938 * cos( 1.69668682344 +       38.13303563780*t);
+   uranus_z_0+=     0.00002578399 * cos( 5.19696447390 +      111.43016149680*t);
+   uranus_z_0+=     0.00002964070 * cos( 6.14338802239 +       35.16409022120*t);
+   uranus_z_0+=     0.00001884104 * cos( 2.61192472648 +       78.71375183040*t);
+   uranus_z_0+=     0.00002330304 * cos( 5.72640226150 +       70.84944530420*t);
+   uranus_z_0+=     0.00001985215 * cos( 0.76408839812 +       39.61750834610*t);
+   uranus_z_0+=     0.00001743154 * cos( 1.20586281789 +      146.59425171800*t);
+   uranus_z_0+=     0.00002037011 * cos( 0.95353587037 +       70.32818044240*t);
+   uranus_z_0+=     0.00001508924 * cos( 2.26195448553 +      222.86032299360*t);
+   uranus_z_0+=     0.00001082736 * cos( 2.62872874057 +      108.46121608020*t);
+   uranus_z_0+=     0.00001051968 * cos( 0.75560753840 +      109.94568878850*t);
+
+   double uranus_z_1=0.0;
+   uranus_z_1+=     0.00655916626 * cos( 0.01271947660 +       74.78159856730*t);
+   uranus_z_1+=     0.00049648951 * cos( 0.00000000000 +        0.00000000000*t);
+   uranus_z_1+=     0.00023874178 * cos( 2.73870491220 +      149.56319713460*t);
+   uranus_z_1+=     0.00007552177 * cos( 5.49304207700 +       76.26607127560*t);
+   uranus_z_1+=     0.00005941304 * cos( 3.61254073304 +       73.29712585900*t);
+   uranus_z_1+=     0.00002868429 * cos( 4.17954157878 +       63.73589830340*t);
+   uranus_z_1+=     0.00002087455 * cos( 5.97858625817 +        1.48447270830*t);
+   uranus_z_1+=     0.00001827697 * cos( 2.71810813335 +       11.04570026390*t);
+   uranus_z_1+=     0.00001305063 * cos( 4.52337002195 +       85.82729883120*t);
+   uranus_z_1+=     0.00001158250 * cos( 5.31913504112 +      224.34479570190*t);
+   uranus_z_1=uranus_z_1 * t;
+
+   double uranus_z_2=0.0;
+   uranus_z_2+=     0.00014697858 * cos( 1.75149165003 +       74.78159856730*t);
+   uranus_z_2+=     0.00001600044 * cos( 3.14159265359 +        0.00000000000*t);
+   uranus_z_2=uranus_z_2 * t * t;
+
+   return uranus_z_0+uranus_z_1+uranus_z_2;
+}
+
+double vsop87a_milli::venus_x(double t){
+   double venus_x_0=0.0;
+   venus_x_0+=     0.72211281391 * cos( 3.17575836361 +    10213.28554621100*t);
+   venus_x_0+=     0.00486448018 * cos( 0.00000000000 +        0.00000000000*t);
+   venus_x_0+=     0.00244500474 * cos( 4.05566613861 +    20426.57109242200*t);
+   venus_x_0+=     0.00002800281 * cos( 0.33147492492 +     2352.86615377180*t);
+   venus_x_0+=     0.00001949669 * cos( 4.23196016801 +     1577.34354244780*t);
+   venus_x_0+=     0.00001241717 * cos( 4.93573787058 +    30639.85663863300*t);
+   venus_x_0+=     0.00001162258 * cos( 2.87958246189 +    18073.70493865020*t);
+   venus_x_0+=     0.00001046690 * cos( 1.75434920413 +     6283.07584999140*t);
+
+   double venus_x_1=0.0;
+   venus_x_1+=     0.00033862636 * cos( 3.14159265359 +        0.00000000000*t);
+   venus_x_1+=     0.00017234992 * cos( 0.92721124604 +    20426.57109242200*t);
+   venus_x_1+=     0.00006510416 * cos( 2.19289889733 +    10213.28554621100*t);
+   venus_x_1=venus_x_1 * t;
+
+   return venus_x_0+venus_x_1;
+}
+
+double vsop87a_milli::venus_y(double t){
+   double venus_y_0=0.0;
+   venus_y_0+=     0.72324820731 * cos( 1.60573808356 +    10213.28554621100*t);
+   venus_y_0+=     0.00549506273 * cos( 3.14159265359 +        0.00000000000*t);
+   venus_y_0+=     0.00244884790 * cos( 2.48564954004 +    20426.57109242200*t);
+   venus_y_0+=     0.00002789807 * cos( 5.04214523606 +     2352.86615377180*t);
+   venus_y_0+=     0.00001933868 * cos( 5.80597990261 +     1577.34354244780*t);
+   venus_y_0+=     0.00001243658 * cos( 3.36573697344 +    30639.85663863300*t);
+   venus_y_0+=     0.00001164480 * cos( 1.30970620277 +    18073.70493865020*t);
+   venus_y_0+=     0.00001041872 * cos( 0.18129136925 +     6283.07584999140*t);
+
+   double venus_y_1=0.0;
+   venus_y_1+=     0.00039231430 * cos( 0.00000000000 +        0.00000000000*t);
+   venus_y_1+=     0.00017282326 * cos( 5.63824735900 +    20426.57109242200*t);
+   venus_y_1+=     0.00005968075 * cos( 3.60854944086 +    10213.28554621100*t);
+   venus_y_1=venus_y_1 * t;
+
+   double venus_y_2=0.0;
+   venus_y_2+=     0.00002007155 * cos( 3.14159265359 +        0.00000000000*t);
+   venus_y_2=venus_y_2 * t * t;
+
+   return venus_y_0+venus_y_1+venus_y_2;
+}
+
+double vsop87a_milli::venus_z(double t){
+   double venus_z_0=0.0;
+   venus_z_0+=     0.04282990302 * cos( 0.26703856476 +    10213.28554621100*t);
+   venus_z_0+=     0.00035588343 * cos( 3.14159265359 +        0.00000000000*t);
+   venus_z_0+=     0.00014501879 * cos( 1.14696911390 +    20426.57109242200*t);
+
+   double venus_z_1=0.0;
+   venus_z_1+=     0.00208096402 * cos( 1.88967278742 +    10213.28554621100*t);
+   venus_z_1+=     0.00001264989 * cos( 3.71037501321 +    20426.57109242200*t);
+   venus_z_1+=     0.00001364144 * cos( 0.00000000000 +        0.00000000000*t);
+   venus_z_1=venus_z_1 * t;
+
+   double venus_z_2=0.0;
+   venus_z_2+=     0.00009148044 * cos( 3.34791005272 +    10213.28554621100*t);
+   venus_z_2=venus_z_2 * t * t;
+
+   return venus_z_0+venus_z_1+venus_z_2;
+}
+
