Added the first version of the Celestial Vault DaySequence plugin

diff --git a/Engine/Plugins/Experimental/CelestialVault/Source/CelestialVault/Public/CelestialDataTypes.h b/Engine/Plugins/Experimental/CelestialVault/Source/CelestialVault/Public/CelestialDataTypes.h
new file mode 100644
--- /dev/null
+++ b/Engine/Plugins/Experimental/CelestialVault/Source/CelestialVault/Public/CelestialDataTypes.h
@@ -0,0 +1,290 @@
+﻿// Copyright Epic Games, Inc. All Rights Reserved.
+
+#pragma once
+
+#include "CoreMinimal.h"
+#include "Engine/DataTable.h"
+#include "CelestialDataTypes.generated.h"
+
+
+/** 
+ * Any celestial body with an Elliptic Orbit type will have to provide the elliptic parameters
+ * Other will use the solar system VSOP87 computations for their location
+ */
+UENUM(BlueprintType) 
+enum class EOrbitType : uint8
+{
+	Elliptic = 0, 
+	Mercury,
+	Venus,
+	Earth,
+	Mars,
+	Jupiter,
+	Saturn,
+	Uranus ,
+	Neptune ,
+	Moon 
+};
+
+/**
+ * TableRow Base type to describe the Creation parameters of a planetary body (Planet, Moon)
+ * Will only be used at creation time, from a proper Data Table
+ */
+USTRUCT(BlueprintType)
+struct FPlanetaryBodyInputData : public FTableRowBase
+{
+	GENERATED_BODY()
+
+	FPlanetaryBodyInputData(const FString& InName, EOrbitType InOrbitType, double InRadius) : Name(InName), OrbitType(InOrbitType), Radius(InRadius) {}
+	FPlanetaryBodyInputData() {}
+public:
+	/** The body Name */
+	UPROPERTY(BlueprintReadWrite, EditAnywhere, Category= "Celestial Vault")
+	FString Name = "";
+
+	/** One of the predefined Orbit types (Solar system planets), any custom Elliptic ones for fantasy planets */
+	UPROPERTY(BlueprintReadWrite, EditAnywhere, Category= "Celestial Vault")
+	EOrbitType OrbitType = EOrbitType::Elliptic;
+	
+	/** The Body Radius in Kilometers */
+	UPROPERTY(BlueprintReadWrite, EditAnywhere, Category= "Celestial Vault", meta=(ForceUnits="Kilometers"))
+	double Radius = 1000.0;
+
+	/** The planetary body Material expects Bodies textures in a single row - This is the 0-based index of the texture to use */
+	UPROPERTY(BlueprintReadWrite, EditAnywhere, Category= "Celestial Vault")
+	int32 TextureColumnIndex = 0;
+
+	// TODO_Beta
+	// Add the Elliptic parameters for other Moon/Planets
+	// Ellipsoid - Minor/Major axes + offset to X axis
+	// Revolution period (year) + angle offset to time
+	// Phase if we want to fake it
+
+public: 
+	static FPlanetaryBodyInputData Earth; // Earth Preset
+	static FPlanetaryBodyInputData Moon; // Moon Preset
+};
+
+/**
+ * Runtime structure to store the computed properties of a Planetary Body, for any Query. 
+ */
+USTRUCT(BlueprintType)
+struct FPlanetaryBodyInfo
+{
+	GENERATED_BODY()
+public:
+	/** One of the predefined body types (Solar system planets), any custom one for fantasy planets */
+	UPROPERTY(BlueprintReadWrite, EditAnywhere, Category= "Celestial Vault")
+	EOrbitType OrbitType = EOrbitType::Elliptic;
+
+	/** The body Name */
+	UPROPERTY(BlueprintReadWrite, EditAnywhere, Category= "Celestial Vault")
+	FString Name = "";
+	
+	/** The Body Radius in Kilometers */
+	UPROPERTY(BlueprintReadWrite, EditAnywhere, Category= "Celestial Vault", meta=(ForceUnits="Kilometers"))
+	double Radius = 1000.0;
+
+	/** The Body Right Ascension in the Celestial Frame - In hours! */
+	UPROPERTY(BlueprintReadWrite, EditAnywhere, Category= "Celestial Vault", meta=(ForceUnits="Hours"))
+	double RA = 0.0;
+
+	/** The Body Declination in the Celestial Frame - In Degrees */
+	UPROPERTY(BlueprintReadWrite, EditAnywhere, Category= "Celestial Vault", meta=(ForceUnits="Degrees"))
+	double DEC  = 0.0;
+
+	/** The Body distance to Earth - In Astronomical Units! */
+	UPROPERTY(BlueprintReadWrite, EditAnywhere, Category= "Celestial Vault")
+	double DistanceInAU = 0.0;
+
+	/** The True apparent diameter of the body seen from Earth (Angular size, in degrees) */
+	UPROPERTY(BlueprintReadWrite, EditAnywhere, Category= "Celestial Vault",  meta = (ForceUnits="Degrees"))
+	double ApparentDiameterDegrees = 0.0;
+
+	/**
+	 * The Scaled apparent diameter of the body seen from Earth (Angular size, in degrees) 
+	 *   Takes the fake scaling factor into consideration, so useful for camera FOV tracking. 
+	 */
+	UPROPERTY(BlueprintReadWrite, EditAnywhere, Category= "Celestial Vault",  meta = (ForceUnits="Degrees"))
+	double ScaledApparentDiameterDegrees = 0.0;
+	
+	/** The Magnitude of the body seen from Earth */
+	UPROPERTY(BlueprintReadWrite, EditAnywhere, Category= "Celestial Vault")
+	double ApparentMagnitude = 0.0;
+
+	/** indication of the crescent phase. 0 = New, 0.5 = Half, 1 = Full*/ 
+	UPROPERTY(BlueprintReadWrite, EditAnywhere, Category= "Celestial Vault")
+	double Phase = 0.0;
+
+	/** Internal value of the ISM Instance corresponding to this planet */
+	UPROPERTY(BlueprintReadWrite, EditAnywhere, Category= "Celestial Vault")
+	int32 ISMInstanceIndex = 0; 
+
+	/** Keep track of the Vault-relative transform, at ToD = 0 for animated bodies (Moon) */
+	UPROPERTY(BlueprintReadWrite, EditAnywhere, Category= "Celestial Vault")
+	FTransform UETransform;
+
+	/** Keep track of the location toward the Earth */ 
+	UPROPERTY(BlueprintReadWrite, EditAnywhere, Category= "Celestial Vault")
+	FVector DirectionTowardEarth = -FVector::ZeroVector;
+
+	// Can't use a UFUNCTION inside a struct --> Use the Celestial Maths Blueprint Function Library
+	FString ToString() const;
+};
+
+/**
+ * Runtime structure to store the computed properties of the Sun, for any Query. 
+ */
+USTRUCT(BlueprintType)
+struct FSunInfo
+{
+	GENERATED_BODY()
+public:
+	/** The Sun Right Ascension in the Celestial Frame - In hours! */
+	UPROPERTY(BlueprintReadWrite, EditAnywhere, Category= "Celestial Vault", meta=(DisplayName="Right Ascension", MakeStructureDefaultValue="0.000000"))
+	double RA = 0.0;
+
+	/** The Sun Declination in the Celestial Frame - In Degrees */
+	UPROPERTY(BlueprintReadWrite, EditAnywhere, Category= "Celestial Vault", meta=(DisplayName="Declination", MakeStructureDefaultValue="0.000000"))
+	double DEC = 0.0;
+
+	/** Keep track of the Vault-relative transform, at ToD = 0 for the Sun */
+	UPROPERTY(BlueprintReadWrite, EditAnywhere, Category= "Celestial Vault")
+	FTransform UETransform;
+
+	/** Keep track of the location toward the Earth */ 
+	UPROPERTY(BlueprintReadWrite, EditAnywhere, Category= "Celestial Vault")
+	FVector DirectionTowardEarth = -FVector::ZeroVector;;
+
+	// Can't use a UFUNCTION inside a struct --> Celestial Maths Blueprint Function Library
+	FString ToString() const;
+
+	// TODO_Beta
+	// Compute and store additional properties for the Sun
+	//   Elevation
+	//   CorrectedElevation - Sun Elevation, corrected for atmospheric diffraction
+	//   Azimuth
+	//   SunriseTime
+	//   SunsetTime
+	//   SolarNoon
+};
+
+/**
+ * Runtime structure to store the computed properties of a Star, for any Query. 
+ */
+USTRUCT(BlueprintType)
+struct FStarInfo
+{
+	GENERATED_BODY()
+
+public:
+	/** The Star Right Ascension in the Celestial Frame - In hours! */
+	UPROPERTY(BlueprintReadWrite, EditAnywhere, Category= "Celestial Vault", meta=(DisplayName="Right Ascension", MakeStructureDefaultValue="0.000000"))
+	double RA = 0.0;
+
+	/** The Star Declination in the Celestial Frame - In Degrees */
+	UPROPERTY(BlueprintReadWrite, EditAnywhere, Category= "Celestial Vault", meta=(DisplayName="Declination", MakeStructureDefaultValue="0.000000"))
+	double DEC = 0.0;
+
+	/** Earth to Star distance (in Parsecs)  */
+	UPROPERTY(BlueprintReadWrite, EditAnywhere, Category= "Celestial Vault", meta=(DisplayName="Distance (in parsecs)", MakeStructureDefaultValue="100.000000"))
+	double DistanceInPC = 100.0;
+	
+	/** Star Name */
+	UPROPERTY(BlueprintReadWrite, EditAnywhere, Category= "Celestial Vault", meta=(DisplayName="Name"))
+	FString Name = "";
+
+	/** Star Magnitude */
+	UPROPERTY(BlueprintReadWrite, EditAnywhere, Category= "Celestial Vault", meta=(DisplayName="Magnitude", MakeStructureDefaultValue="1.000000"))
+	double Magnitude = 1.0;
+
+	/** Star RGB Color - Can be computed from the B-V value if the star is from an official Catalog */
+	UPROPERTY(BlueprintReadWrite, EditAnywhere, Category= "Celestial Vault", meta=(DisplayName="Color"))
+	FLinearColor Color = FLinearColor(1.0f,1.0f,1.0f,1.0f);
+
+	/** Star Hipparcos ID if present in the Hipparcos Catalog */
+	UPROPERTY(BlueprintReadWrite, EditAnywhere, Category= "Celestial Vault", meta=(DisplayName="Hipparcos Catalog ID", MakeStructureDefaultValue="0"))
+	int32 HipparcosID = 0;
+
+	/** Star Henry Draper ID if present in the Henry Draper Catalog */
+	UPROPERTY(BlueprintReadWrite, EditAnywhere, Category= "Celestial Vault", meta=(DisplayName="Henry Draper Catalog ID", MakeStructureDefaultValue="0"))
+	int32 HenryDraperID = 0;
+
+	/** Star YaleBrightStar ID if present in the Yale Bright Star Catalog */
+	UPROPERTY(BlueprintReadWrite, EditAnywhere, Category= "Celestial Vault", meta=(DisplayName="Yale Bright Star Catalog ID", MakeStructureDefaultValue="0"))
+	int32 YaleBrightStarID = 0;
+	
+	/** Star Color Index, also named B-V */
+	UPROPERTY(BlueprintReadWrite, EditAnywhere, Category= "Celestial Vault", meta=(DisplayName="ColorIndex", MakeStructureDefaultValue="0.000000"))
+	double ColorIndex = 0.0;
+
+	/** Internal value of the ISM Instance corresponding to this Star */
+	UPROPERTY(BlueprintReadWrite, EditAnywhere, Category= "Celestial Vault",  meta = (HideInDataTable))
+	int32 ISMInstanceIndex = 0;
+
+	// Can't use a UFUNCTION inside a struct --> Celestial Maths Blueprint Function Library
+	FString ToString() const;
+};
+
+
+/**
+ * TableRow Base type to describe the Creation parameters of a Basic Star
+ * Will only be used at creation time, from a proper Data Table
+ * This Struct contains the minimal needed data for a Fictional Star
+ */
+USTRUCT(BlueprintType)
+struct FStarInputData : public FTableRowBase 
+{
+	GENERATED_BODY()
+public:
+	/** The Star Right Ascension in the Celestial Frame - In hours! */
+	UPROPERTY(BlueprintReadWrite, EditAnywhere, Category= "Celestial Vault", meta=(DisplayName="Right Ascension", MakeStructureDefaultValue="0.000000"))
+	double RA = 0.0;
+
+	/** The Star Declination in the Celestial Frame - In Degrees */
+	UPROPERTY(BlueprintReadWrite, EditAnywhere, Category= "Celestial Vault", meta=(DisplayName="Declination", MakeStructureDefaultValue="0.000000"))
+	double DEC = 0.0;
+
+	/** Earth to Star distance (in Parsecs)  */
+	UPROPERTY(BlueprintReadWrite, EditAnywhere, Category= "Celestial Vault", meta=(DisplayName="Distance (in parsecs)", MakeStructureDefaultValue="100.000000"))
+	double DistanceInPC = 100.0;
+	
+	/** Star Name */
+	UPROPERTY(BlueprintReadWrite, EditAnywhere, Category= "Celestial Vault", meta=(DisplayName="Name"))
+	FString Name = "";
+
+	/** Star Magnitude */
+	UPROPERTY(BlueprintReadWrite, EditAnywhere, Category= "Celestial Vault", meta=(DisplayName="Magnitude", MakeStructureDefaultValue="1.000000"))
+	double Magnitude = 1.0;
+
+	/** Star RGB Color - Useless if computed from the B-V */ 
+	UPROPERTY(BlueprintReadWrite, EditAnywhere, Category= "Celestial Vault", meta=(DisplayName="Color"))
+	FLinearColor Color = FLinearColor(1.0f,1.0f,1.0f,1.0f);
+};
+
+/**
+ * TableRow Base type to describe the Creation parameters of a Catalog-based Star
+ * Will only be used at creation time, from a proper Data Table
+ * This Struct extends the FStarInputData class with additional Catalog Properties
+ */
+USTRUCT(BlueprintType)
+struct FCelestialStarInputData : public FStarInputData 
+{
+	GENERATED_BODY()
+public:
+	/** Star Hipparcos ID */
+	UPROPERTY(BlueprintReadWrite, EditAnywhere, Category= "Celestial Vault", meta=(DisplayName="Hipparcos Catalog ID", MakeStructureDefaultValue="0"))
+	int32 HipparcosID = -1;
+
+	/** Star Henry Draper ID */
+	UPROPERTY(BlueprintReadWrite, EditAnywhere, Category= "Celestial Vault", meta=(DisplayName="Henry Draper Catalog ID", MakeStructureDefaultValue="0"))
+	int32 HenryDraperID = -1;
+
+	/** Star YaleBrightStar ID */
+	UPROPERTY(BlueprintReadWrite, EditAnywhere, Category= "Celestial Vault", meta=(DisplayName="Yale Bright Star Catalog ID", MakeStructureDefaultValue="0"))
+	int32 YaleBrightStarID = -1;
+	
+	/** Star Color Index, also named B-V */
+	UPROPERTY(BlueprintReadWrite, EditAnywhere, Category= "Celestial Vault", meta=(DisplayName="ColorIndex", MakeStructureDefaultValue="0.000000"))
+	double ColorIndex = -1.0;
+};
\ No newline at end of file
diff --git a/Engine/Plugins/Experimental/CelestialVault/Source/CelestialVault/Public/CelestialVaultDaySequenceActor.h b/Engine/Plugins/Experimental/CelestialVault/Source/CelestialVault/Public/CelestialVaultDaySequenceActor.h
new file mode 100644
--- /dev/null
+++ b/Engine/Plugins/Experimental/CelestialVault/Source/CelestialVault/Public/CelestialVaultDaySequenceActor.h
@@ -0,0 +1,249 @@
+﻿// Copyright Epic Games, Inc. All Rights Reserved.
+
+#pragma once
+
+#include "CoreMinimal.h"
+#include "CelestialDataTypes.h"
+#include "DaySequenceActor.h"
+#include "CelestialVaultDaySequenceActor.generated.h"
+
+// Celestial Classes
+class UStarsDataTable;
+
+// UE Classes
+class USkyAtmosphereComponent;
+class USkyLightComponent;
+class UVolumetricCloudComponent;
+class UDirectionalLightComponent;
+//class UExponentialHeightFogComponent;
+class UStaticMeshComponent;
+class UInstancedStaticMeshComponent;
+
+
+UCLASS(Blueprintable, HideCategories=(Tags, Networking, LevelInstance))
+class CELESTIALVAULT_API ACelestialVaultDaySequenceActor
+	: public ADaySequenceActor
+{
+	GENERATED_BODY()
+
+public:
+	ACelestialVaultDaySequenceActor(const FObjectInitializer& Init);
+ 
+public: // Properties
+
+#pragma region Components 
+	
+	UPROPERTY(VisibleDefaultsOnly, BlueprintReadOnly, Category= "Celestial Vault", meta = (AllowPrivateAccess = "true"))
+	TObjectPtr<USceneComponent> PlanetCenterComponent;
+
+	UPROPERTY(VisibleDefaultsOnly, BlueprintReadOnly, Category= "Celestial Vault", meta = (AllowPrivateAccess = "true"))
+	TObjectPtr<USceneComponent> CelestialVaultComponent;
+
+	UPROPERTY(VisibleDefaultsOnly, BlueprintReadOnly, Category= "Celestial Vault", meta = (AllowPrivateAccess = "true"))
+	TObjectPtr<UDirectionalLightComponent> SunLightComponent;
+
+	UPROPERTY(VisibleDefaultsOnly, BlueprintReadOnly, Category= "Celestial Vault", meta = (AllowPrivateAccess = "true"))
+	TObjectPtr<UDirectionalLightComponent> MoonLightComponent;
+	
+	UPROPERTY(VisibleDefaultsOnly, BlueprintReadOnly, Category = "Celestial Vault", meta = (AllowPrivateAccess = "true"))
+	TObjectPtr<USkyAtmosphereComponent> SkyAtmosphereComponent;
+	
+	UPROPERTY(VisibleDefaultsOnly, BlueprintReadOnly, Category = "Celestial Vault", meta = (AllowPrivateAccess = "true"))
+	TObjectPtr<USkyLightComponent> SkyLightComponent;
+	
+	//TODO - Should HeightFog be optional because it doesn't work round earth when far from the origin?
+	// UPROPERTY(VisibleDefaultsOnly, BlueprintReadOnly, Category = "Celestial Vault", meta = (AllowPrivateAccess = "true"))
+	// TObjectPtr<UExponentialHeightFogComponent> ExponentialHeightFogComponent;
+	
+	UPROPERTY(VisibleDefaultsOnly, BlueprintReadOnly, Category = "Celestial Vault", meta = (AllowPrivateAccess = "true"))
+	TObjectPtr<UVolumetricCloudComponent> VolumetricCloudComponent;
+
+	UPROPERTY(VisibleDefaultsOnly, BlueprintReadOnly, Category = "Celestial Vault", meta = (AllowPrivateAccess = "true"))
+	TObjectPtr<UStaticMeshComponent> DeepSkyComponent;
+
+	UPROPERTY(VisibleDefaultsOnly, BlueprintReadOnly, Category = "Celestial Vault", meta = (AllowPrivateAccess = "true"))
+	TObjectPtr<UStaticMeshComponent> MoonDiscComponent;
+
+	UPROPERTY(VisibleDefaultsOnly, BlueprintReadOnly, Category = "Celestial Vault", meta = (AllowPrivateAccess = "true"))
+	TObjectPtr<UInstancedStaticMeshComponent> StarsComponent;
+
+	UPROPERTY(VisibleDefaultsOnly, BlueprintReadOnly, Category = "Celestial Vault", meta = (AllowPrivateAccess = "true"))
+	TObjectPtr<UInstancedStaticMeshComponent> PlanetsComponent;
+
+#pragma endregion
+
+#pragma region DateTime related Properties
+
+	/* If true, ignore the Year Month Day value and use the current system Date */ 
+	UPROPERTY(EditAnywhere, BlueprintReadWrite, Category="Celestial Vault|Date & Location")
+	bool bUseCurrentDate = false;
+
+	/** Current Year*/ 
+	UPROPERTY(EditAnywhere, BlueprintReadWrite, Category="Celestial Vault|Date & Location", meta = (EditConditionHides, EditCondition = "bUseCurrentDate==false", ClampMin = "1", ClampMax = "9999"))
+	int Year = 2025;
+
+	/** Current Month*/ 
+	UPROPERTY(EditAnywhere, BlueprintReadWrite, Category="Celestial Vault|Date & Location", meta = (EditConditionHides, EditCondition = "bUseCurrentDate==false", ClampMin = "1", ClampMax = "12"))
+	int Month = 1;
+
+	/** Current Day*/ 
+	UPROPERTY(EditAnywhere, BlueprintReadWrite, Category="Celestial Vault|Date & Location", meta = (EditConditionHides, EditCondition = "bUseCurrentDate==false", ClampMin = "1", ClampMax = "31"))
+	int Day = 7;
+
+	/** Current Time Zone*/ 
+	UPROPERTY(EditAnywhere, BlueprintReadWrite, Category="Celestial Vault|Date & Location", meta = (ClampMin = "-11", ClampMax = "14"))
+	double GMT_TimeZone = -5.0;
+
+	/** Set to true if your current date is during the DST period */ 
+	UPROPERTY(EditAnywhere, BlueprintReadWrite, Category = "Celestial Vault|Date & Location")
+	bool bIsDaylightSavings = false;
+	
+	/** Latitude of Level Origin on planet */
+	UPROPERTY(EditAnywhere, BlueprintReadWrite, Category = "Celestial Vault|Date & Location", meta = (ForceUnits="Degrees", ClampMin = "-90", ClampMax = "90"))
+	double Latitude = 45.0;
+
+	/** Longitude of Level Origin on planet */
+	UPROPERTY(EditAnywhere, BlueprintReadWrite, Category = "Celestial Vault|Date & Location", meta = (ForceUnits="Degrees", ClampMin = "-180", ClampMax = "180"))
+	double Longitude = -73.0;
+	
+	// Greenwich Mean Sidereal Time at corresponding to a 0 Time of Day (midnight in the morning) for the selected Date. 
+	UPROPERTY(BlueprintReadOnly, VisibleAnywhere, Category="Celestial Vault|Date & Location")
+	double GMST_AtTOD_0 = 0.0;
+
+	// Transform to apply to the planet to have it located tangent to the Origin
+	UPROPERTY(BlueprintReadOnly, VisibleAnywhere, Category="Celestial Vault|Date & Location")
+	FTransform PlanetCenterTransform;
+	
+#pragma endregion
+
+#pragma region Celestial Vault related Properties
+
+	/** We generate the sky elements the "Platon" way, using a sphere surrounding the Earth. This is the radius of this sphere. Make sure it's not too small to avoid parallax effects */ 
+	UPROPERTY(EditAnywhere, BlueprintReadWrite, AdvancedDisplay, Category="Celestial Vault", meta=( ForceUnits="Kilometers", ClampMin = "0", ClampMax = "500000"))
+	double CelestialVaultDistance = 400000.0;
+
+	/** Percentage of the CelestialVaultDistance at which the Stars are created */ 
+	UPROPERTY(EditAnywhere, BlueprintReadWrite, AdvancedDisplay, Category="Celestial Vault", meta=( ForceUnits="Percent", ClampMin = "0.1", ClampMax = "100"))
+	double StarsVaultPercentage = 99.0;
+
+	/** Percentage of the CelestialVaultDistance at which the Planets are created */
+	UPROPERTY(EditAnywhere, BlueprintReadWrite, AdvancedDisplay, Category="Celestial Vault", meta=( ForceUnits="Percent", ClampMin = "0.1", ClampMax = "100"))
+	double PlanetsVaultPercentage = 97.0;
+	
+	/** Percentage of the CelestialVaultDistance at which the Moons are created */
+	UPROPERTY(EditAnywhere, BlueprintReadWrite, AdvancedDisplay, Category="Celestial Vault", meta=( ForceUnits="Percent", ClampMin = "0.1", ClampMax = "100"))
+	double MoonVaultPercentage = 95.0;
+	
+#pragma endregion
+
+#pragma region Stars related Properties
+	
+	/** A Datatable containing a Celestial Star Catalog data */
+	UPROPERTY(EditAnywhere, BlueprintReadWrite, Category="Celestial Vault|Stars")
+	TObjectPtr<UDataTable> CelestialStarCatalog = nullptr;
+
+	/** A Datatable containing a Fictional Star Catalog data */
+	UPROPERTY(EditAnywhere, BlueprintReadWrite, Category="Celestial Vault|Stars")
+	TObjectPtr<UDataTable> FictionalStarCatalog = nullptr;
+	
+	/** All stars from the catalog with a Magnitude dimmer than this threshold won't be generated - Usually 6 is the naked eye visibility limit*/ 
+	UPROPERTY(EditAnywhere, BlueprintReadWrite, Category="Celestial Vault|Stars", meta = (ClampMin = "-30", ClampMax = "30"))
+	float MaxVisibleMagnitude = 6.0f;
+
+	/** If true, the Stars information will be kept in memory and queryable at runtime */ 
+	UPROPERTY(EditAnywhere, BlueprintReadWrite, Category="Celestial Vault|Stars")
+	bool KeepStarsInfo = false;
+	
+	/** Array of the created Stars information - Only populated if KeepStarsInfo is true */ 
+	UPROPERTY(BlueprintReadOnly, Category="Celestial Vault|Stars")
+	TArray<FStarInfo> StarsInfo;
+	
+#pragma endregion
+	
+#pragma region Planets related Properties
+
+	/** The Data Catalog containing all Planets data */
+	UPROPERTY(EditAnywhere, BlueprintReadWrite, Category="Celestial Vault|Planets")
+	TObjectPtr<UDataTable> PlanetsCatalog = nullptr;
+
+	/** Factor to artificially increase the Planetary bodies size */ 
+	UPROPERTY(EditAnywhere, BlueprintReadWrite, Category="Celestial Vault|Planets", meta = (ClampMin = "0"))
+	float PlanetsScale = 1.0f;
+
+	/** If true, the Stars information will be kept in memory and queryable at runtime */ 
+	UPROPERTY(EditAnywhere, BlueprintReadWrite, Category="Celestial Vault|Planets")
+	bool KeepPlanetsInfos = false;
+	
+	/** Array of the created planetary bodies, with all their computed information - Only populated if KeepPlanetsInfos is true */ 
+	UPROPERTY(BlueprintReadOnly, Category="Celestial Vault|Planets")
+	TArray<FPlanetaryBodyInfo> PlanetsInfos;
+
+#pragma endregion
+
+#pragma region Moon / Sun related Properties
+
+	/** Factor to artificially increase the Moon size */
+	UPROPERTY(EditAnywhere, BlueprintReadWrite, Category="Celestial Vault|Moons", meta = (ClampMin = "0"))
+	float MoonScale = 1.0f;
+
+	/** Celestial Info for the Moon, at the beginning of the Day */  
+	UPROPERTY(BlueprintReadOnly, Category="Celestial Vault|Moons")
+	FPlanetaryBodyInfo MoonBodyInfo;
+
+	/** Celestial Info for the Sun, at the beginning of the Day */  
+	UPROPERTY(BlueprintReadOnly, Category="Celestial Vault|Sun")
+	FSunInfo SunInfo; 
+
+#pragma endregion 
+
+protected: // Functions
+
+	/** BeginPlay and OnConstruction overrides auto-register this actor with the DaySequenceSubsystem. */
+	virtual void BeginPlay() override;
+	virtual void OnConstruction(const FTransform& Transform) override;
+	
+public: // Functions
+
+	/** Returns the current defined day, without any Time, because the Time will be controlled by the DaySequence Time of day - Uses "Now" or the Year/Month/Day properties */ 
+	UFUNCTION(BlueprintCallable, Category = "Celestial Vault|Date & Location")
+	FDateTime GetDate() const;
+
+	/** Returns the Celestial Info for the Sun, at a specific Julian Date */  
+	UFUNCTION(BlueprintCallable, Category = "Celestial Vault|Sun")
+	FSunInfo GetSunInfo(double JulianDate);
+	
+	/** Returns the Celestial Info for the Moon, at a specific Julian Date */  
+	UFUNCTION(BlueprintCallable, Category = "Celestial Vault|Queries")
+	FPlanetaryBodyInfo GetMoonInfo(double JulianDate);
+
+// Data Queries
+	
+	/** Return the Celestial Information of the Star closest to a specific direction, within an angle threshold */
+	UFUNCTION(BlueprintCallable, Category = "Celestial Vault|Queries")
+	bool GetClosestStarInfo(FVector ObserverLocation, FVector LookupDirection, double ThresholdAngleDegree, FStarInfo& FoundStarInfo, FTransform& StarTransform);
+	
+	/** Return the Celestial Information of the Planetary Body (moon, planet) closest to a specific direction, within an angle threshold */
+	UFUNCTION(BlueprintCallable, Category = "Celestial Vault|Queries")
+	bool GetClosestPlanetaryBody(FVector StartPosition, FVector LookupDirection, double ThresholdAngleDegree, FPlanetaryBodyInfo& FoundPlanetaryBodyInfo, FTransform& BodyTransform);
+
+	/** Return the Celestial Information of a specific Planetary Body (by its orbit type) */
+	UFUNCTION(BlueprintCallable, Category = "Celestial Vault|Queries")
+	bool GetPlanetaryBodyByOrbitType(EOrbitType OrbitType, FPlanetaryBodyInfo& FoundPlanetaryBodyInfo, FTransform& BodyTransform);
+
+	/** Returns the Celestial vault rotation angle for the current Date at t=0 midnight */
+	double GetDayCelestialVaultAngle() const;
+
+#if WITH_EDITOR
+	virtual void PostEditChangeProperty(FPropertyChangedEvent& PropertyChangedEvent) override;
+#endif
+
+private:
+	void InitStars();
+	void InitPlanetaryBodies();
+	FPlanetaryBodyInfo GetPlanetaryBodyInfo(const FPlanetaryBodyInputData& InputPlanetaryBody, double JulianDay, double UEDistance, double BodyScale) const;
+
+	UFUNCTION(BlueprintCallable, CallInEditor, Category= "Celestial Vault")
+	void RebuildAll();
+
+	FDelegateHandle DrawDebugDelegateHandle;
+};
diff --git a/Engine/Plugins/Experimental/CelestialVault/Source/CelestialVault/Public/CelestialVaultSequence.h b/Engine/Plugins/Experimental/CelestialVault/Source/CelestialVault/Public/CelestialVaultSequence.h
new file mode 100644
--- /dev/null
+++ b/Engine/Plugins/Experimental/CelestialVault/Source/CelestialVault/Public/CelestialVaultSequence.h
@@ -0,0 +1,27 @@
+﻿// Copyright Epic Games, Inc. All Rights Reserved.
+
+#pragma once
+
+#include "CoreMinimal.h"
+#include "ProceduralDaySequence.h"
+
+#include "CelestialVaultSequence.generated.h"
+
+
+/**
+ * A procedural sequence that animates a sun in a physically accurate way based on geographic data.
+ */
+USTRUCT()
+struct CELESTIALVAULT_API FCelestialVaultSequence : public FProceduralDaySequence
+{
+	GENERATED_BODY()
+
+	virtual ~FCelestialVaultSequence() override
+	{}
+
+	UPROPERTY(EditAnywhere, Category = "Procedural Parameters")
+	unsigned KeyCount = 24;
+	
+private:
+	virtual void BuildSequence(UProceduralDaySequenceBuilder* InBuilder) override;
+};
diff --git a/Engine/Plugins/Experimental/CelestialVault/Source/CelestialVault/Public/vsop87a_full_EarthMoon.h b/Engine/Plugins/Experimental/CelestialVault/Source/CelestialVault/Public/vsop87a_full_EarthMoon.h
new file mode 100644
--- /dev/null
+++ b/Engine/Plugins/Experimental/CelestialVault/Source/CelestialVault/Public/vsop87a_full_EarthMoon.h
@@ -0,0 +1,23 @@
+// Copyright Epic Games, Inc. All Rights Reserved.
+
+//VSOP87-Multilang http://www.celestialprogramming.com/vsop87-multilang/index.html
+//Greg Miller ([REDACTED_EMAIL]) 2019.  Released as Public Domain
+
+#ifndef VSOP87A_FULL
+#define VSOP87A_FULL
+
+class vsop87a_full{
+   public:
+   static void getEarth(double t,double temp[]);
+   static void getEmb(double t,double temp[]);
+   static void getMoon(double earth[], double emb[],double temp[]);
+
+   private:
+   static double earth_x(double t);
+   static double earth_y(double t);
+   static double earth_z(double t);
+   static double emb_x(double t);
+   static double emb_y(double t);
+   static double emb_z(double t);
+};
+#endif
diff --git a/Engine/Plugins/Experimental/CelestialVault/Source/CelestialVault/Public/vsop87a_milli.h b/Engine/Plugins/Experimental/CelestialVault/Source/CelestialVault/Public/vsop87a_milli.h
new file mode 100644
--- /dev/null
+++ b/Engine/Plugins/Experimental/CelestialVault/Source/CelestialVault/Public/vsop87a_milli.h
@@ -0,0 +1,51 @@
+// Copyright Epic Games, Inc. All Rights Reserved.
+
+//VSOP87-Multilang http://www.celestialprogramming.com/vsop87-multilang/index.html
+//Greg Miller ([REDACTED_EMAIL]) 2019.  Released as Public Domain
+
+#ifndef VSOP87A_MILLI
+#define VSOP87A_MILLI
+
+class vsop87a_milli{
+   public:
+   static void getEarth(double t,double temp[]);
+   static void getEmb(double t,double temp[]);
+   static void getJupiter(double t,double temp[]);
+   static void getMars(double t,double temp[]);
+   static void getMercury(double t,double temp[]);
+   static void getNeptune(double t,double temp[]);
+   static void getSaturn(double t,double temp[]);
+   static void getUranus(double t,double temp[]);
+   static void getVenus(double t,double temp[]);
+   static void getMoon(double earth[], double emb[],double temp[]);
+
+   private:
+   static double earth_x(double t);
+   static double earth_y(double t);
+   static double earth_z(double t);
+   static double emb_x(double t);
+   static double emb_y(double t);
+   static double emb_z(double t);
+   static double jupiter_x(double t);
+   static double jupiter_y(double t);
+   static double jupiter_z(double t);
+   static double mars_x(double t);
+   static double mars_y(double t);
+   static double mars_z(double t);
+   static double mercury_x(double t);
+   static double mercury_y(double t);
+   static double mercury_z(double t);
+   static double neptune_x(double t);
+   static double neptune_y(double t);
+   static double neptune_z(double t);
+   static double saturn_x(double t);
+   static double saturn_y(double t);
+   static double saturn_z(double t);
+   static double uranus_x(double t);
+   static double uranus_y(double t);
+   static double uranus_z(double t);
+   static double venus_x(double t);
+   static double venus_y(double t);
+   static double venus_z(double t);
+};
+#endif
