Added missing copyright boilerplate

diff --git a/Engine/Plugins/Experimental/CelestialVault/Source/CelestialVaultEditor/CelestialVaultEditor.Build.cs b/Engine/Plugins/Experimental/CelestialVault/Source/CelestialVaultEditor/CelestialVaultEditor.Build.cs
--- a/Engine/Plugins/Experimental/CelestialVault/Source/CelestialVaultEditor/CelestialVaultEditor.Build.cs
+++ b/Engine/Plugins/Experimental/CelestialVault/Source/CelestialVaultEditor/CelestialVaultEditor.Build.cs
@@ -1,4 +1,6 @@
-﻿using UnrealBuildTool;
+﻿// Copyright Epic Games, Inc. All Rights Reserved.
+
+using UnrealBuildTool;
 
 public class CelestialVaultEditor : ModuleRules
 {
diff --git a/Engine/Plugins/Experimental/CelestialVault/Source/CelestialVaultEditor/Private/CelestialVaultEditor.cpp b/Engine/Plugins/Experimental/CelestialVault/Source/CelestialVaultEditor/Private/CelestialVaultEditor.cpp
--- a/Engine/Plugins/Experimental/CelestialVault/Source/CelestialVaultEditor/Private/CelestialVaultEditor.cpp
+++ b/Engine/Plugins/Experimental/CelestialVault/Source/CelestialVaultEditor/Private/CelestialVaultEditor.cpp
@@ -1,4 +1,6 @@
-﻿#include "CelestialVaultEditor.h"
+﻿// Copyright Epic Games, Inc. All Rights Reserved.
+
+#include "CelestialVaultEditor.h"
 
 #define LOCTEXT_NAMESPACE "FCelestialVaultEditorModule"
 
diff --git a/Engine/Plugins/Experimental/CelestialVault/Source/CelestialVaultEditor/Public/CelestialVaultEditor.h b/Engine/Plugins/Experimental/CelestialVault/Source/CelestialVaultEditor/Public/CelestialVaultEditor.h
--- a/Engine/Plugins/Experimental/CelestialVault/Source/CelestialVaultEditor/Public/CelestialVaultEditor.h
+++ b/Engine/Plugins/Experimental/CelestialVault/Source/CelestialVaultEditor/Public/CelestialVaultEditor.h
@@ -1,4 +1,6 @@
-﻿#pragma once
+﻿// Copyright Epic Games, Inc. All Rights Reserved.
+
+#pragma once
 
 #include "CoreMinimal.h"
 #include "Modules/ModuleManager.h"
