Added the Moon Phases
- Compute the Age and Illumination Percentage
- With manual control of the Moon Phase to override the Ephemeris values
- Added Intensity base values for SunLight and Moonlight
- Wrote a proper shader for the Moon plane imposter, replicating the moon sphere, craters, and unlit face

Improvements
- Moved the Celestial Body transformation computation to the BodyInfo class (needed for Manual Control)
- Added some ForceUnits meta properties

#virtualized

diff --git a/Engine/Plugins/Experimental/CelestialVault/Source/CelestialVault/Private/CelestialDataTypes.cpp b/Engine/Plugins/Experimental/CelestialVault/Source/CelestialVault/Private/CelestialDataTypes.cpp
--- a/Engine/Plugins/Experimental/CelestialVault/Source/CelestialVault/Private/CelestialDataTypes.cpp
+++ b/Engine/Plugins/Experimental/CelestialVault/Source/CelestialVault/Private/CelestialDataTypes.cpp
@@ -27,7 +27,8 @@ FString FPlanetaryBodyInfo::ToString() const
 			"Radius: %.0f km\n"
 			"Apparent Diameter: %s / %.8f°\n"
 			"Apparent Magnitude: %.2f\n"
-			"Phase: %.2f\n"
+			"Age: %.2f (%.2f days) since last New Moon\n"
+			"IlluminationPercentage: %.2f\n"
 			"-------\n"
 			"Scaled Apparent Diameter: %.8f°\n"
 			),
@@ -37,15 +38,38 @@ FString FPlanetaryBodyInfo::ToString() const
 			*UCelestialMaths::Conv_DeclinationToString(DEC),
 			DistanceInAU,
 			Radius,
-			*UCelestialMaths::Conv_DeclinationToString(ApparentDiameterDegrees),
-			ApparentDiameterDegrees,
+			*UCelestialMaths::Conv_DeclinationToString(ApparentDiameterDegrees), ApparentDiameterDegrees,
 			ApparentMagnitude,
-			Phase, 
+			Age, UCelestialMaths::SynodicMonthAverage * Age, 
+			IlluminationPercentage,
 			ScaledApparentDiameterDegrees
 		); 
 	
 }
 
+void FPlanetaryBodyInfo::ComputeTranform(double UEDistance, double BodyScale)
+{
+	// Location
+	FVector BodyLocation = UCelestialMaths::RADECToXYZ_RH(RA * 15, DEC, UEDistance);
+	BodyLocation.Y *= -1; // Convert to UE Frame by inverting Y
+
+	// Rotation - Orient the Body mesh towards the earth -
+	// Use Look at the UE Origin, this is an acceptable approximation if the vault is big enough - We could use the Earth center if needed
+	FVector BodyToEarth = (FVector::ZeroVector - BodyLocation).GetSafeNormal();
+	FQuat BodyRotation = FQuat(BodyToEarth.Rotation());
+
+	// Scale
+	// The mesh plane is 100 UE Units (1m), and located at PlanetMeshDistanceKm. Use Thales theorem to compute its effective scale at this distance to ensure the right apparent diameter
+	double UERadius = (UEDistance * Radius * 1000.0) / UCelestialMaths::AstronomicalUnitsToMeters(DistanceInAU);
+	double Scale = UERadius * BodyScale / 50.0; // The plane half-length is 50 UE Units
+	FVector BodyScale3D = FVector(Scale, Scale, Scale);
+
+	UETransform = FTransform(BodyRotation, BodyLocation, BodyScale3D);
+
+	// Direction
+	DirectionTowardEarth = BodyToEarth;
+}
+
 FString FSunInfo::ToString() const
 {
 	return FString::Printf(TEXT(
diff --git a/Engine/Plugins/Experimental/CelestialVault/Source/CelestialVault/Private/CelestialMaths.cpp b/Engine/Plugins/Experimental/CelestialVault/Source/CelestialVault/Private/CelestialMaths.cpp
--- a/Engine/Plugins/Experimental/CelestialVault/Source/CelestialVault/Private/CelestialMaths.cpp
+++ b/Engine/Plugins/Experimental/CelestialVault/Source/CelestialVault/Private/CelestialMaths.cpp
@@ -31,9 +31,12 @@ FMatrix UCelestialMaths::VSOPToJ2000 = FMatrix(
 		FVector(+0.000000000000, +0.397776982902, +0.917482137087),		 
 		FVector(+0.000000000000, 0.000000000000, +0.000000000000)).GetTransposed();	// No Origin offset
 
-double UCelestialMaths::SpeedOfLightMetersPerSeconds = 299792458;
+double UCelestialMaths::SpeedOfLightMetersPerSeconds = 299792458.0;
 
-double UCelestialMaths::AstronomicalUnitsMeters = 149597870700;
+double UCelestialMaths::AstronomicalUnitsMeters = 149597870700.0;
+
+double UCelestialMaths::NewMoonReferenceJulianDate = 2460705.025196759; // Known new Moon was January 29th, 2025, at 12:36:17 UTC
+double UCelestialMaths::SynodicMonthAverage = 29.530588853;
 
 #pragma endregion
 
@@ -44,12 +47,12 @@ FLinearColor UCelestialMaths::BVtoLinearColor(float BV)
 	// From https://en.wikipedia.org/wiki/Color_index
 
 	// Model valid only between [-0.4, 2.0] 
-	BV = FMath::Clamp(BV, -0.4, 2.0);
+	BV = FMath::Clamp(BV, -0.4f, 2.0f);
 
 	// Compute Effective Temperature
-	const double K1 = 0.92 * BV + 1.7;
-	const double K2 = 0.92 * BV + 0.62;
-	double Temperature = 4600.0 * (1.0 / K1 + 1.0 / K2);
+	const float K1 = 0.92f * BV + 1.7f;
+	const float K2 = 0.92f * BV + 0.62f;
+	float Temperature = 4600.0f * (1.0f / K1 + 1.0f / K2);
 
 	// Convert to Color
 	return FLinearColor::MakeFromColorTemperature( Temperature);;
@@ -75,7 +78,7 @@ FVector UCelestialMaths::GetBodyLocation_FK5J2000_AU(FPlanetaryBodyInputData Pla
 
 	
 	// Get the body coordinates in the inertial frame defined by the dynamical equinox and ecliptic J2000 (JD2451545.0).
-	double BodyXYZ[3] = { 0, 0, 0 }; 
+	double BodyXYZ[3] = { 0.0, 0.0, 0.0 }; 
 	switch (PlanetaryBody.OrbitType)
 	{
 	case EOrbitType::Mercury:
@@ -105,9 +108,9 @@ FVector UCelestialMaths::GetBodyLocation_FK5J2000_AU(FPlanetaryBodyInputData Pla
 	case EOrbitType::Moon:
 		{
 			// Special case for the Moon. VSOP Works by getting the Earth/Moon barycenter, and combine it with the earth XYZ
-			double EarthXYZ[3] = { 0, 0, 0 }; 
+			double EarthXYZ[3] = { 0.0, 0.0, 0.0 }; 
 			vsop87a_full::getEarth(Time, EarthXYZ);
-			double EMBXYZ[3] = { 0, 0, 0 };
+			double EMBXYZ[3] = { 0.0, 0.0, 0.0 };
 			vsop87a_full::getEmb(Time, EMBXYZ); 
 			vsop87a_full::getMoon(EarthXYZ, EMBXYZ, BodyXYZ);	
 		}
@@ -144,9 +147,9 @@ void UCelestialMaths::GetBodyCelestialCoordinatesAU(double JulianDate, FPlanetar
 	TRACE_CPUPROFILER_EVENT_SCOPE(UCelestialMaths::GetBodyCelestialCoordinatesAU);
 
 	
-	RAHours = 0;
-	DECDegrees = 0;
-	DistanceBodyToEarthAU = 0; 
+	RAHours = 0.0;
+	DECDegrees = 0.0;
+	DistanceBodyToEarthAU = 0.0; 
 	
 	if (PlanetaryBody.OrbitType == EOrbitType::Earth)
 	{
@@ -173,7 +176,7 @@ void UCelestialMaths::GetBodyCelestialCoordinatesAU(double JulianDate, FPlanetar
 	FVector TopoCentricTargetAU = BodyLocationGeoCentricAU - ObserverLocationGeocentricAU;
 	DistanceEarthToSunAU = (ObserverLocationGeocentricAU-EarthLocationFK5J2000AU).Length();
 	
-	double RADegrees = 0;
+	double RADegrees = 0.0;
 	XYZToRADEC_RH(TopoCentricTargetAU, RADegrees, DECDegrees, DistanceBodyToEarthAU);
 	RAHours = RADegrees / 15.0;
 
@@ -199,21 +202,21 @@ double UCelestialMaths::GetPlanetaryBodyMagnitude(FPlanetaryBodyInputData Planet
 	//   ΣnCn αn is called the phase function 
 	
 	
-	double DistanceFactor = 5 * FMath::LogX(10, DistanceToEarthAU * DistanceToSunAU);
-	double ApparentMagnitude = 0; 
-	PhaseAngle = FMath::RadiansToDegrees(FMath::Acos( (DistanceToSunAU * DistanceToSunAU + DistanceToEarthAU * DistanceToEarthAU - DistanceEarthToSunAU*DistanceEarthToSunAU) / (2 * DistanceToSunAU * DistanceToEarthAU) ));
+	double DistanceFactor = 5.0 * FMath::LogX(10.0, DistanceToEarthAU * DistanceToSunAU);
+	double ApparentMagnitude = 0.0; 
+	PhaseAngle = FMath::RadiansToDegrees(FMath::Acos( (DistanceToSunAU * DistanceToSunAU + DistanceToEarthAU * DistanceToEarthAU - DistanceEarthToSunAU*DistanceEarthToSunAU) / (2.0 * DistanceToSunAU * DistanceToEarthAU) ));
 	double PhaseAngle2 = PhaseAngle * PhaseAngle;
 	double PhaseAngle3 = PhaseAngle2 * PhaseAngle;
 	double PhaseAngle4 = PhaseAngle3 * PhaseAngle;
 	double PhaseAngle5 = PhaseAngle4 * PhaseAngle;
 	double PhaseAngle6 = PhaseAngle5 * PhaseAngle;
-	double PhaseFunction = 0; 
+	double PhaseFunction = 0.0; 
 	
 	switch (PlanetaryBody.OrbitType)
 	{
 	case EOrbitType::Elliptic:
 		// Fictional Body. Return dummy magnitude
-		return 0;
+		return 0.0;
 		break;
 	
 	case EOrbitType::Mercury:
@@ -237,7 +240,7 @@ double UCelestialMaths::GetPlanetaryBodyMagnitude(FPlanetaryBodyInputData Planet
 		PhaseFunction = -1.060E-3 * PhaseAngle + 2.054E-4 * PhaseAngle2;
 		break;
 	case EOrbitType::Mars:
-		if ( PhaseAngle < 50)
+		if ( PhaseAngle < 50.0)
 		{
 			ApparentMagnitude = -1.601;
 			PhaseFunction =  0.02267 * PhaseAngle - 0.0001302 * PhaseAngle2;
@@ -249,7 +252,7 @@ double UCelestialMaths::GetPlanetaryBodyMagnitude(FPlanetaryBodyInputData Planet
 		}
 		break;
 	case EOrbitType::Jupiter:
-		if ( PhaseAngle < 12)
+		if ( PhaseAngle < 12.0)
 		{
 			ApparentMagnitude =  -9.395;
 			PhaseFunction =  -3.7E-04 * PhaseAngle - 6.16E-04 * PhaseAngle2;
@@ -257,18 +260,18 @@ double UCelestialMaths::GetPlanetaryBodyMagnitude(FPlanetaryBodyInputData Planet
 		else // 12 < α < 130 - The phase curve of Jupiter as seen from Earth cannot exceed α = 12 so we should be good. Add this one just in case...
 		{
 			ApparentMagnitude = -9.428;
-			PhaseFunction = -2.5 * FMath::LogX(10,
+			PhaseFunction = -2.5 * FMath::LogX(10.0,
 				1.0
 				-1.507 * (PhaseAngle/180.0)
-				-0.363 * FMath::Pow(PhaseAngle / 180.0, 2)
-				-0.062 * FMath::Pow(PhaseAngle / 180.0, 3)
-				+2.809 * FMath::Pow(PhaseAngle / 180.0, 4)
-				-1.876 * FMath::Pow(PhaseAngle / 180.0, 5)); 
+				-0.363 * FMath::Pow(PhaseAngle / 180.0, 2.0)
+				-0.062 * FMath::Pow(PhaseAngle / 180.0, 3.0)
+				+2.809 * FMath::Pow(PhaseAngle / 180.0, 4.0)
+				-1.876 * FMath::Pow(PhaseAngle / 180.0, 5.0)); 
 		}
 		break;
 	case EOrbitType::Saturn:
 		// Keep it simple and ignore the ring effects
-		if ( PhaseAngle < 6)
+		if ( PhaseAngle < 6.0)
 		{
 			ApparentMagnitude =  -8.95;
 			PhaseFunction =  - 3.7E-04 * PhaseAngle + 6.16E-04 * PhaseAngle2;
@@ -281,14 +284,13 @@ double UCelestialMaths::GetPlanetaryBodyMagnitude(FPlanetaryBodyInputData Planet
 		break;
 	case EOrbitType::Uranus:
 		ApparentMagnitude = -7.19;
-		PhaseFunction =  -8.4E-04 * 82;  // PhaseAngle doesn't have any impact, and 82 is the most important planetographic latitude
+		PhaseFunction =  -8.4E-04 * 82.0;  // PhaseAngle doesn't have any impact, and 82 is the most important planetographic latitude
 		break;
 	case EOrbitType::Neptune:
 		ApparentMagnitude = -7.00;
 		PhaseFunction =  7.944E-3 * PhaseAngle + 9.617E-5 * PhaseAngle2;
 		break;
 	case EOrbitType::Moon:
-		// TODO Alban - A verifier
 		ApparentMagnitude = -12.73;
 		PhaseFunction = 0.026 * PhaseAngle + 4E-9 * PhaseAngle4;
 		break;
@@ -298,6 +300,20 @@ double UCelestialMaths::GetPlanetaryBodyMagnitude(FPlanetaryBodyInputData Planet
 	return DistanceFactor + ApparentMagnitude + PhaseFunction;
 }
 
+double UCelestialMaths::GetMoonNormalizedAgeSimple(double JulianDate)
+{
+	double DeltaDays = JulianDate - NewMoonReferenceJulianDate;
+	double MoonAgeDays = ModPositive(DeltaDays, SynodicMonthAverage);
+
+	return MoonAgeDays / SynodicMonthAverage;
+}
+
+double UCelestialMaths::GetIlluminationPercentage(double NormalizedAge)
+{
+	return 0.5 * (1.0-FMath::Cos(2.0*UE_PI * NormalizedAge));
+}
+
+
 #pragma endregion
 
 #pragma region Time
@@ -357,7 +373,7 @@ double UCelestialMaths::UTCDateTimeToJulianDate(FDateTime UTCDateTime)
 	}
 
 	// Compute the Julian Date
-	double JulianDate = FloorForJulianDate(365.25 * (Year + 4716)) + FloorForJulianDate(30.6001 * (Month + 1)) + Day + b - 1524.5;
+	double JulianDate = FloorForJulianDate(365.25 * (Year + 4716.0)) + FloorForJulianDate(30.6001 * (Month + 1)) + Day + b - 1524.5;
 	JulianDate += Hours / 24.0;
 	JulianDate += Minutes / 24.0 / 60.0;
 	JulianDate += Seconds / 24.0 / 60.0 / 60.0;
@@ -378,7 +394,7 @@ FDateTime UCelestialMaths::JulianDateToUTCDateTime(double JulianDate)
 	if (IntegralPart >= 2299161)
 	{
 		int32 Alpha = FloorForJulianDate((IntegralPart - 1867216.25) / 36524.25);
-		A = IntegralPart + 1 + Alpha - FloorForJulianDate(Alpha / 4);
+		A = IntegralPart + 1 + Alpha - FloorForJulianDate(Alpha / 4.0);
 	}
 
 	// Compute ABCDE values
@@ -387,14 +403,16 @@ FDateTime UCelestialMaths::JulianDateToUTCDateTime(double JulianDate)
 	int32 D = FloorForJulianDate(365.25 * C);
 	int32 E = FloorForJulianDate((B - D) / 30.6001);
 
-	int32 day=B-D-static_cast<int32>(30.6001 * E)+FractionalPart;
-	int32 month=E-1;
-	if(E>13){
-		month=E-13;
+	int32 day = B - D - static_cast<int32>(30.6001 * E) + FractionalPart;
+	int32 month = E - 1;
+	if (E > 13)
+	{
+		month = E - 13;
 	}
-	int32 year=C-4716;
-	if(month<3){
-		year=C-4715;
+	int32 year = C - 4716;
+	if (month < 3)
+	{
+		year = C - 4715;
 	}
 
 	// Split in H,M,S,MS
@@ -414,7 +432,7 @@ double UCelestialMaths::DateTimeToGreenwichMeanSiderealTime(FDateTime UTCDateTim
 	double JulianDate = UTCDateTime.GetJulianDay();
 	double JulianCentury = JulianDateToJulianCenturies(JulianDate);
 	double SideralTimeAt0h = 100.46061837 + 36000.770053608 * JulianCentury + 0.000387933 * JulianCentury * JulianCentury + 1.0/38710000.0 * JulianCentury * JulianCentury * JulianCentury;
-	return ModPositive(SideralTimeAt0h, 360);
+	return ModPositive(SideralTimeAt0h, 360.0);
 }
 
 double UCelestialMaths::JulianDateToGreenwichMeanSiderealTime(double JulianDate) {
@@ -453,7 +471,7 @@ double UCelestialMaths::JulianDateToGreenwichApparentSiderealTime(double JulianD
 	const double EE = EquationOfTheEquinoxes(JulianDate);
 	double gast = GMST + EE;
 
-	return ModPositive(gast,360);
+	return ModPositive(gast,360.0);
 }
 
 double UCelestialMaths::JulianDateToJulianCenturies(double JulianDate)
@@ -530,7 +548,7 @@ double UCelestialMaths::JulianDateToInternationalAtomicTime(double JulianDate)
 double UCelestialMaths::ModPositive(double Value, double Modulo)
 {
 	double Result = FMath::Fmod(Value, Modulo);
-	if (Result < 0)
+	if (Result < 0.0)
 	{
 		Result += Modulo;
 	}
@@ -539,13 +557,13 @@ double UCelestialMaths::ModPositive(double Value, double Modulo)
 
 void UCelestialMaths::DegreesToHMS(double DecimalDegrees, int32& Hours, int32& Minutes, double& Seconds)
 {
-	DecimalDegrees = FMath::Fmod(DecimalDegrees, 360);
-	if (DecimalDegrees < 0)
+	DecimalDegrees = FMath::Fmod(DecimalDegrees, 360.0);
+	if (DecimalDegrees < 0.0)
 	{
-		DecimalDegrees += 360;
+		DecimalDegrees += 360.0;
 	}
 
-	double AngleHours = DecimalDegrees/15.0;
+	double AngleHours = DecimalDegrees / 15.0;
 
 	Hours = StaticCast<uint32>(AngleHours);
 	Minutes = StaticCast<uint32>((AngleHours - Hours) * 60.0);
@@ -555,14 +573,14 @@ void UCelestialMaths::DegreesToHMS(double DecimalDegrees, int32& Hours, int32& M
 void UCelestialMaths::DegreesToDMS(double DecimalDegrees, bool& Sign, int32& Degrees, int32& Minutes, double& Seconds)
 {
 	Sign = true;
-	DecimalDegrees = FMath::Fmod(DecimalDegrees, 360);
+	DecimalDegrees = FMath::Fmod(DecimalDegrees, 360.0);
 
-	if (DecimalDegrees < 0)
+	if (DecimalDegrees < 0.0)
 	{
 		Sign = false;
-		DecimalDegrees *= -1;
+		DecimalDegrees *= -1.0;
 	}
-	
+
 	Degrees = StaticCast<uint32>(DecimalDegrees);
 	Minutes = StaticCast<uint32>((DecimalDegrees - Degrees) * 60.0);
 	Seconds = (DecimalDegrees - Degrees) * 3600.0 - Minutes * 60.0;
@@ -586,7 +604,7 @@ double UCelestialMaths::GetEarthRotationAngle(double JulianDate)
 	const double JulianDateFraction = FMath::Fractional(JulianDate);
 	const double Rotations = (0.779057273264 + 0.00273781191135448 * DU + JulianDateFraction);
 
-	return ModPositive(Rotations * 360, 360);
+	return ModPositive(Rotations * 360.0, 360.0);
 }
 
 FVector UCelestialMaths::GeodeticLatLonToECEFXYZAU(double Latitude, double Longitude, double Altitude)
@@ -596,11 +614,11 @@ FVector UCelestialMaths::GeodeticLatLonToECEFXYZAU(double Latitude, double Longi
 	double LongitudeRadians = FMath::DegreesToRadians(Longitude);
 	
 	const double a = MetersToAstronomicalUnits(6378136.6);
-	const double f = 1 / 298.25642;
+	const double f = 1.0 / 298.25642;
 
-	const double C = 1 / FMath::Sqrt(FMath::Cos(LatitudeRadians) * FMath::Cos(LatitudeRadians) + (1.0 - f) * (1.0 - f) * (FMath::Sin(LatitudeRadians) * FMath::Sin(LatitudeRadians)));
+	const double C = 1.0 / FMath::Sqrt(FMath::Cos(LatitudeRadians) * FMath::Cos(LatitudeRadians) + (1.0 - f) * (1.0 - f) * (FMath::Sin(LatitudeRadians) * FMath::Sin(LatitudeRadians)));
 
-	const double S = (1 - f) * (1 - f) * C;
+	const double S = (1.0 - f) * (1.0 - f) * C;
 	const double h = MetersToAstronomicalUnits(Altitude);
 
 	return FVector(
@@ -679,39 +697,39 @@ void UCelestialMaths::Nutation2000BTruncated(double JulianDate, double& DeltaPsi
 
 	//Terms summed from lowest to highest to reduce floating point rounding errors.  See coefficients Page 88.
 	// Constants are first multiplied by 10000000 to reduce errors
-	double DeltaPsiArcseconds = 0;
-	double DeltaEpsilonArcSeconds = 0;
+	double DeltaPsiArcSeconds = 0.0;
+	double DeltaEpsilonArcSeconds = 0.0;
 
 	// FundamentalArgument #6
-	double FundamentalArgument = Lp + 2 * (F - D + Omega);
-	DeltaPsiArcseconds += (-516821 + 1226 * JC) * FMath::Sin(FundamentalArgument) + -524 * FMath::Cos(FundamentalArgument);
-	DeltaEpsilonArcSeconds += (224386 + -677 * JC) * FMath::Cos(FundamentalArgument) + -174 * FMath::Sin(FundamentalArgument);
+	double FundamentalArgument = Lp + 2.0 * (F - D + Omega);
+	DeltaPsiArcSeconds += (-516821.0 + 1226.0 * JC) * FMath::Sin(FundamentalArgument) + -524.0 * FMath::Cos(FundamentalArgument);
+	DeltaEpsilonArcSeconds += (224386.0 + -677.0 * JC) * FMath::Cos(FundamentalArgument) + -174.0 * FMath::Sin(FundamentalArgument);
 
 	// FundamentalArgument #5 - Just Lp
-	DeltaPsiArcseconds += (1475877 + -3633 * JC) * FMath::Sin(Lp) + 11817 * FMath::Cos(Lp);
-	DeltaEpsilonArcSeconds += (73871 + -184 * JC) * FMath::Cos(Lp) + -1924 * FMath::Sin(Lp);
+	DeltaPsiArcSeconds += (1475877.0 + -3633.0 * JC) * FMath::Sin(Lp) + 11817.0 * FMath::Cos(Lp);
+	DeltaEpsilonArcSeconds += (73871.0 + -184.0 * JC) * FMath::Cos(Lp) + -1924.0 * FMath::Sin(Lp);
 
 	// FundamentalArgument #4
-	FundamentalArgument = 2 * Omega;
-	DeltaPsiArcseconds += (2074554 + 207 * JC) * FMath::Sin(FundamentalArgument) + -698 * FMath::Cos(FundamentalArgument);
-	DeltaEpsilonArcSeconds += (-897492 + 470 * JC) * FMath::Cos(FundamentalArgument) + -291 * FMath::Sin(FundamentalArgument);
+	FundamentalArgument = 2.0 * Omega;
+	DeltaPsiArcSeconds += (2074554.0 + 207.0 * JC) * FMath::Sin(FundamentalArgument) + -698.0 * FMath::Cos(FundamentalArgument);
+	DeltaEpsilonArcSeconds += (-897492.0 + 470.0 * JC) * FMath::Cos(FundamentalArgument) + -291.0 * FMath::Sin(FundamentalArgument);
 
 	// FundamentalArgument #3
-	FundamentalArgument = 2 * (F + Omega);
-	DeltaPsiArcseconds += (-2276413 + -234 * JC) * FMath::Sin(FundamentalArgument) + 2796 * FMath::Cos(FundamentalArgument);
-	DeltaEpsilonArcSeconds += (978459 + -485 * JC) * FMath::Cos(FundamentalArgument) + 1374 * FMath::Sin(FundamentalArgument);
+	FundamentalArgument = 2.0 * (F + Omega);
+	DeltaPsiArcSeconds += (-2276413.0 + -234.0 * JC) * FMath::Sin(FundamentalArgument) + 2796.0 * FMath::Cos(FundamentalArgument);
+	DeltaEpsilonArcSeconds += (978459.0 + -485.0 * JC) * FMath::Cos(FundamentalArgument) + 1374.0 * FMath::Sin(FundamentalArgument);
 
 	// FundamentalArgument #2
-	FundamentalArgument = 2 * (F - D + Omega);
-	DeltaPsiArcseconds += (-13170906 + -1675 * JC) * FMath::Sin(FundamentalArgument) + -13696 * FMath::Cos(FundamentalArgument);
-	DeltaEpsilonArcSeconds += (5730336 + -3015 * JC) * FMath::Cos(FundamentalArgument) + -4587 * FMath::Sin(FundamentalArgument);
+	FundamentalArgument = 2.0 * (F - D + Omega);
+	DeltaPsiArcSeconds += (-13170906.0 + -1675.0 * JC) * FMath::Sin(FundamentalArgument) + -13696.0 * FMath::Cos(FundamentalArgument);
+	DeltaEpsilonArcSeconds += (5730336.0 + -3015.0 * JC) * FMath::Cos(FundamentalArgument) + -4587.0 * FMath::Sin(FundamentalArgument);
 
 	// FundamentalArgument #1
-	DeltaPsiArcseconds += (-172064161 + -174666 * JC) * FMath::Sin(Omega) + 33386 * FMath::Cos(Omega);
-	DeltaEpsilonArcSeconds += (92052331 + 9086 * JC) * FMath::Cos(Omega) + 15377 * FMath::Sin(Omega);
+	DeltaPsiArcSeconds += (-172064161.0 + -174666.0 * JC) * FMath::Sin(Omega) + 33386.0 * FMath::Cos(Omega);
+	DeltaEpsilonArcSeconds += (92052331.0 + 9086.0 * JC) * FMath::Cos(Omega) + 15377.0 * FMath::Sin(Omega);
 
-	DeltaPsi = ArcsecondsToDegrees(DeltaPsiArcseconds / 10000000);
-	DeltaEpsilon = ArcsecondsToDegrees(DeltaEpsilonArcSeconds / 10000000);
+	DeltaPsi = ArcsecondsToDegrees(DeltaPsiArcSeconds / 10000000.0);
+	DeltaEpsilon = ArcsecondsToDegrees(DeltaEpsilonArcSeconds / 10000000.0);
 }
 
 FTransform UCelestialMaths::GetPlanetCenterTransform(double Latitude, double Longitude, double Altitude)
@@ -746,7 +764,7 @@ FTransform UCelestialMaths::GetPlanetCenterTransform(double Latitude, double Lon
 	}
 	else
 	{
-		double Tolerance = 1.e-50; // Normalize with a very low threshold, because default is 10-8, too high for double computations
+		double Tolerance = 1.E-50; // Normalize with a very low threshold, because default is 10-8, too high for double computations
 
 		// Compute the ellipsoid normal (Earth...)
 		double SemiMajorMeter = 6378137.0;
@@ -795,27 +813,27 @@ FSunInfo UCelestialMaths::GetSunInformation(double JulianDate, double ObserverLa
 		// 	TimeOffset += 1.0;
 		// }
 
-		double LatitudeRad = FMath::DegreesToRadians(ObserverLatitude);
+	    double LatitudeRad = FMath::DegreesToRadians(ObserverLatitude);
 
-		// Get the julian day (number of days since Jan 1st of the year 4713 BC)
-		double JulianDay = JulianDate;
-		double JulianCentury = JulianDateToJulianCenturies(JulianDate);
+	    // Get the julian day (number of days since Jan 1st of the year 4713 BC)
+	    double JulianDay = JulianDate;
+	    double JulianCentury = JulianDateToJulianCenturies(JulianDate);
 
-		// Get the sun's mean longitude , referred to the mean equinox of julian date
-		double GeomMeanLongSunDeg = FMath::Fmod(280.46646 + JulianCentury*(36000.76983 + JulianCentury*0.0003032), 360.0);
-		double GeomMeanLongSunRad = FMath::DegreesToRadians(GeomMeanLongSunDeg);
+	    // Get the sun's mean longitude , referred to the mean equinox of julian date
+	    double GeomMeanLongSunDeg = FMath::Fmod(280.46646 + JulianCentury * (36000.76983 + JulianCentury * 0.0003032), 360.0);
+	    double GeomMeanLongSunRad = FMath::DegreesToRadians(GeomMeanLongSunDeg);
 
-		// Get the sun's mean anomaly
-		double GeomMeanAnomSunDeg = 357.52911 + JulianCentury*(35999.05029 - 0.0001537*JulianCentury);
-		double GeomMeanAnomSunRad = FMath::DegreesToRadians(GeomMeanAnomSunDeg);
+	    // Get the sun's mean anomaly
+	    double GeomMeanAnomSunDeg = 357.52911 + JulianCentury * (35999.05029 - 0.0001537 * JulianCentury);
+	    double GeomMeanAnomSunRad = FMath::DegreesToRadians(GeomMeanAnomSunDeg);
 
-		// Get the earth's orbit eccentricity
-		double EccentEarthOrbit = 0.016708634 - JulianCentury*(0.000042037 + 0.0000001267*JulianCentury);
-		
-		// Get the sun's equation of the center
-		double SunEqOfCtr = FMath::Sin(GeomMeanAnomSunRad)*(1.914602 - JulianCentury*(0.004817 + 0.000014*JulianCentury))
-						  + FMath::Sin(2.0 * GeomMeanAnomSunRad)*(0.019993 - 0.000101*JulianCentury)
-						  + FMath::Sin(3.0 * GeomMeanAnomSunRad)*0.000289;
+	    // Get the earth's orbit eccentricity
+	    double EccentEarthOrbit = 0.016708634 - JulianCentury * (0.000042037 + 0.0000001267 * JulianCentury);
+
+	    // Get the sun's equation of the center
+	    double SunEqOfCtr = FMath::Sin(GeomMeanAnomSunRad) * (1.914602 - JulianCentury * (0.004817 + 0.000014 * JulianCentury))
+		    + FMath::Sin(2.0 * GeomMeanAnomSunRad) * (0.019993 - 0.000101 * JulianCentury)
+		    + FMath::Sin(3.0 * GeomMeanAnomSunRad) * 0.000289;
 		
 		// Get the sun's true longitude
 		double SunTrueLongDeg = GeomMeanLongSunDeg + SunEqOfCtr;
@@ -930,8 +948,8 @@ void UCelestialMaths::XYZToRADEC_RH(FVector XYZ, double& RADegrees, double& DECD
 	double t = FMath::Acos(XYZ.Z / Radius);
 
 	//Make sure RA is positive, and Dec is in range +/-90
-	if (l < 0) {
-		l += 2 * UE_PI ;
+	if (l < 0.0) {
+		l += 2.0 * UE_PI ;
 	}
 	t = 0.5 * UE_PI - t;
 
@@ -1014,7 +1032,7 @@ FString UCelestialMaths::Conv_DMSToString(bool Sign, int32 Degrees, int32 Minute
 
 int32 UCelestialMaths::FloorForJulianDate(double JulianDate)
 {
-	if (JulianDate > 0)
+	if (JulianDate > 0.0)
 	{
 		return StaticCast<int32>(FMath::Floor(JulianDate));
 	}
diff --git a/Engine/Plugins/Experimental/CelestialVault/Source/CelestialVault/Private/CelestialVaultDaySequenceActor.cpp b/Engine/Plugins/Experimental/CelestialVault/Source/CelestialVault/Private/CelestialVaultDaySequenceActor.cpp
--- a/Engine/Plugins/Experimental/CelestialVault/Source/CelestialVault/Private/CelestialVaultDaySequenceActor.cpp
+++ b/Engine/Plugins/Experimental/CelestialVault/Source/CelestialVault/Private/CelestialVaultDaySequenceActor.cpp
@@ -15,6 +15,7 @@
 
 // UE Objects
 #include "Materials/MaterialInterface.h"
+#include "Materials/MaterialInstanceDynamic.h"
 #include "UObject/ConstructorHelpers.h"
 #include "DaySequenceSubsystem.h"
 #include "DaySequenceCollectionAsset.h"
@@ -28,7 +29,6 @@
 
 
 
-
 // Sets default values
 ACelestialVaultDaySequenceActor::ACelestialVaultDaySequenceActor(const FObjectInitializer& Init)
 : Super(Init)
@@ -117,23 +117,26 @@ ACelestialVaultDaySequenceActor::ACelestialVaultDaySequenceActor(const FObjectIn
 	MoonDiscComponent->SetAffectDynamicIndirectLighting(false);
 	MoonDiscComponent->SetCanEverAffectNavigation(false);
 	MoonDiscComponent->SetMaterial(0, MoonDiscDefaultMaterial.Object.Get());
+	
 
 	// Moon - Light
 	MoonLightComponent = CreateDefaultSubobject<UDirectionalLightComponent>(TEXT("Moon Light"));
 	MoonLightComponent->SetupAttachment(CelestialVaultComponent);
 	MoonLightComponent->SetAtmosphereSunLightIndex(1);	// Make Moon the secondary directional light that contributes to the sky atmosphere.
 	MoonLightComponent->SetForwardShadingPriority(0); // Give Moon forward shading priority.
-	MoonLightComponent->SetIntensity(0.32f);
+	MoonLightComponent->SetIntensity(MoonLightIntensity);
 	MoonLightComponent->SetUseTemperature(true);
 	MoonLightComponent->SetTemperature(9000.f);
-	MoonLightComponent->SetWorldRotation(FRotator(-45.f, 0.f, 0.f));
+	MoonLightComponent->SetWorldRotation(FRotator(-45.f, 0.0f, 0.0f));
+	MoonLightComponent->bCastCloudShadows = true; // Otherwise we still have hard shadows with an overcast sky 
 	
 	// Attach the sunlight relative to the rotating Vault
 	SunLightComponent = CreateDefaultSubobject<UDirectionalLightComponent>(TEXT("Sun Light"));
 	SunLightComponent->SetupAttachment(CelestialVaultComponent);
 	SunLightComponent->SetAtmosphereSunLightIndex(0);	// Make Sun the first directional light that contributes to the sky atmosphere.
 	SunLightComponent->SetForwardShadingPriority(1);   // Give Sun forward shading priority.
-	SunLightComponent->SetIntensity(120000.0f);
+	SunLightComponent->SetIntensity(SunLightIntensity);
+	SunLightComponent->bCastCloudShadows = true; // Otherwise we still have hard shadows with an overcast sky
 	
 	// Sequence and Data Assets
 	if (!IsTemplate())
@@ -183,6 +186,14 @@ void ACelestialVaultDaySequenceActor::OnConstruction(const FTransform& Transform
 			}
 		}
 
+		// Replace the Moon Material with a MID
+		UMaterialInstanceDynamic * MoonDiscMaterial = Cast<UMaterialInstanceDynamic>(MoonDiscComponent->GetMaterial(0));
+		if (!MoonDiscMaterial)
+		{
+			MoonDiscMaterial = MoonDiscComponent->CreateAndSetMaterialInstanceDynamic(0); // Make the material Dynamic to control the Phase
+		}
+		MoonDiscMaterial->SetScalarParameterValue(FName("MoonAge"), MoonAge);
+		
 		RebuildAll();  
 	}
 }
@@ -210,8 +221,8 @@ FSunInfo ACelestialVaultDaySequenceActor::GetSunInfo(double JulianDate)
 	FSunInfo ResultSunInfo = UCelestialMaths::GetSunInformation(JulianDate, Latitude, Longitude);
 
 	// Location
-	FVector SunLocation = UCelestialMaths::RADECToXYZ_RH(ResultSunInfo.RA * 15, ResultSunInfo.DEC, 1000); // TODO - Add the Sun Distance To SunInfo 
-	SunLocation *= FVector(100, -100, 100); // Convert to UE Frame by inverting Y and scaling to UE Units - // TODO - use AU
+	FVector SunLocation = UCelestialMaths::RADECToXYZ_RH(ResultSunInfo.RA * 15.0, ResultSunInfo.DEC, 1000.0); // TODO - Add the Sun Distance To SunInfo 
+	SunLocation.Y *= -100.0; // Convert to UE Frame by inverting Y and scaling to UE Units
 	ResultSunInfo.UETransform.SetLocation(SunLocation);
 	ResultSunInfo.DirectionTowardEarth = (FVector::ZeroVector - SunLocation).GetSafeNormal();
 	return ResultSunInfo;
@@ -219,8 +230,46 @@ FSunInfo ACelestialVaultDaySequenceActor::GetSunInfo(double JulianDate)
 
 FPlanetaryBodyInfo ACelestialVaultDaySequenceActor::GetMoonInfo(double JulianDate)
 {
-	double UEDistance = CelestialVaultDistance * 100000 * MoonVaultPercentage / 100.0;
-	return GetPlanetaryBodyInfo(FPlanetaryBodyInputData::Moon, JulianDate, UEDistance, MoonScale);
+	double UEDistance = CelestialVaultDistance * 100000.0 * MoonVaultPercentage / 100.0;
+	FPlanetaryBodyInfo ResultMoonInfo = GetPlanetaryBodyInfo(FPlanetaryBodyInputData::Moon, JulianDate, UEDistance, MoonScale);
+
+	if (bManualControl)
+	{
+		// Phase/Age
+		ResultMoonInfo.Age = MoonAge;
+		ResultMoonInfo.IlluminationPercentage = UCelestialMaths::GetIlluminationPercentage(MoonAge);
+
+		// Location
+		FSunInfo SunInfoTemp = GetSunInfo(JulianDate);
+		ResultMoonInfo.RA = UCelestialMaths::ModPositive(SunInfoTemp.RA + MoonOffset_RA, 24.0);
+		ResultMoonInfo.DEC = UCelestialMaths::ModPositive(SunInfoTemp.DEC + MoonOffset_DEC, 360.0);
+		
+		// We changed the RA, so we need to update the transform
+		ResultMoonInfo.ComputeTranform(UEDistance, MoonScale);
+	}
+	else
+	{
+		ResultMoonInfo.Age = UCelestialMaths::GetMoonNormalizedAgeSimple(JulianDate);;
+		ResultMoonInfo.IlluminationPercentage = UCelestialMaths::GetIlluminationPercentage(ResultMoonInfo.Age);
+	}
+	return ResultMoonInfo;
+}
+
+void ACelestialVaultDaySequenceActor::SetMoonDiscAge(float InMoonAge)
+{
+	if (MoonDiscComponent)
+	{
+		UMaterialInstanceDynamic * MoonDiscMaterial = Cast<UMaterialInstanceDynamic>(MoonDiscComponent->GetMaterial(0));
+		if (MoonDiscMaterial)
+		{
+			MoonDiscMaterial->SetScalarParameterValue(FName("MoonAge"), InMoonAge);
+		}
+	}
+
+	if (MoonLightComponent)
+	{
+		MoonLightComponent->SetIntensity(MoonLightIntensity * UCelestialMaths::GetIlluminationPercentage(InMoonAge));
+	}
 }
 
 bool ACelestialVaultDaySequenceActor::GetClosestStarInfo(FVector ObserverLocation, FVector LookupDirection, double ThresholdAngleDegrees, FStarInfo& FoundStarInfo, FTransform& StarTransform)
@@ -228,7 +277,7 @@ bool ACelestialVaultDaySequenceActor::GetClosestStarInfo(FVector ObserverLocatio
 	StarTransform = FTransform::Identity;
 	
 	// We query only if we generated the StarInfo  
-	if (!KeepStarsInfo || !StarsComponent)
+	if (!bKeepStarsInfo || !StarsComponent)
 	{
 		return false;
 	}
@@ -270,7 +319,7 @@ bool ACelestialVaultDaySequenceActor::GetClosestPlanetaryBody(FVector ObserverPo
 	BodyTransform = FTransform::Identity;
 	
 	// We query only if we generated the StarInfo  
-	if (!KeepPlanetsInfos)
+	if (!bKeepPlanetsInfos)
 	{
 		return false;
 	}
@@ -339,7 +388,7 @@ bool ACelestialVaultDaySequenceActor::GetPlanetaryBodyByOrbitType(EOrbitType Orb
 	BodyTransform = FTransform::Identity;
 	
 	// We query only if we generated the StarInfo  
-	if (!KeepPlanetsInfos)
+	if (!bKeepPlanetsInfos)
 	{
 		return false;
 	}
@@ -386,10 +435,10 @@ void ACelestialVaultDaySequenceActor::PostEditChangeProperty(FPropertyChangedEve
 	//Get the name of the property that was changed  
 	FName PropertyName = (PropertyChangedEvent.Property != nullptr) ? PropertyChangedEvent.Property->GetFName() : NAME_None;
 
-	bool BuildAll = false;
-	bool BuildStars = false;
-	bool BuildPlanetaryBodies = false;
-	bool RebuildSequence = false;
+	bool bBuildAll = false;
+	bool bBuildStars = false;
+	bool bBuildPlanetaryBodies = false;
+	bool bRebuildSequence = false;
 
 	// Geometry properties
 	if (PropertyName == GET_MEMBER_NAME_CHECKED(ACelestialVaultDaySequenceActor, CelestialVaultDistance) ||
@@ -404,8 +453,8 @@ void ACelestialVaultDaySequenceActor::PostEditChangeProperty(FPropertyChangedEve
 		
 		// TODO UpdateGeoReferencing();
 
-		BuildAll |= true;
-		RebuildSequence = true;
+		bBuildAll |= true;
+		bRebuildSequence = true;
 	}
 
 	// Time properties
@@ -422,43 +471,49 @@ void ACelestialVaultDaySequenceActor::PostEditChangeProperty(FPropertyChangedEve
 		FDateTime DateTime = GetDate();
 		GMST_AtTOD_0 = UCelestialMaths::DateTimeToGreenwichMeanSiderealTime(DateTime);
 		double CelestialVaultAngle = GetDayCelestialVaultAngle();
-		CelestialVaultComponent->SetRelativeRotation( FRotator(0, CelestialVaultAngle, 0.0) );
+		CelestialVaultComponent->SetRelativeRotation( FRotator(0.0, CelestialVaultAngle, 0.0) );
 
-		BuildAll = true;
-		RebuildSequence = true;
+		bBuildAll = true;
+		bRebuildSequence = true;
 	}
 
 	// Properties impacting the Stars
 	if (PropertyName == GET_MEMBER_NAME_CHECKED(ACelestialVaultDaySequenceActor, CelestialStarCatalog) ||
 		PropertyName == GET_MEMBER_NAME_CHECKED(ACelestialVaultDaySequenceActor, FictionalStarCatalog) ||
 		PropertyName == GET_MEMBER_NAME_CHECKED(ACelestialVaultDaySequenceActor, MaxVisibleMagnitude) ||
-		PropertyName == GET_MEMBER_NAME_CHECKED(ACelestialVaultDaySequenceActor, KeepStarsInfo))
+		PropertyName == GET_MEMBER_NAME_CHECKED(ACelestialVaultDaySequenceActor, bKeepStarsInfo))
 	{
-		BuildStars |= true;
-		RebuildSequence = true;
+		bBuildStars |= true;
+		bRebuildSequence = true;
 	}
 
 	// Properties impacting the Planetary Bodies
 	if (PropertyName == GET_MEMBER_NAME_CHECKED(ACelestialVaultDaySequenceActor, PlanetsCatalog) ||
 		PropertyName == GET_MEMBER_NAME_CHECKED(ACelestialVaultDaySequenceActor, PlanetsScale) ||
-		PropertyName == GET_MEMBER_NAME_CHECKED(ACelestialVaultDaySequenceActor, KeepPlanetsInfos) ||
-		PropertyName == GET_MEMBER_NAME_CHECKED(ACelestialVaultDaySequenceActor, MoonScale))
+		PropertyName == GET_MEMBER_NAME_CHECKED(ACelestialVaultDaySequenceActor, bKeepPlanetsInfos) ||
+		PropertyName == GET_MEMBER_NAME_CHECKED(ACelestialVaultDaySequenceActor, MoonScale)||
+		PropertyName == GET_MEMBER_NAME_CHECKED(ACelestialVaultDaySequenceActor, bManualControl)||
+		PropertyName == GET_MEMBER_NAME_CHECKED(ACelestialVaultDaySequenceActor, MoonAge) ||
+		PropertyName == GET_MEMBER_NAME_CHECKED(ACelestialVaultDaySequenceActor, MoonLightIntensity) || 
+		PropertyName == GET_MEMBER_NAME_CHECKED(ACelestialVaultDaySequenceActor, SunLightIntensity) || 
+		PropertyName == GET_MEMBER_NAME_CHECKED(ACelestialVaultDaySequenceActor, MoonOffset_RA) ||
+		PropertyName == GET_MEMBER_NAME_CHECKED(ACelestialVaultDaySequenceActor, MoonOffset_DEC))
 	{
-		BuildPlanetaryBodies |= true;
-		RebuildSequence = true;
+		bBuildPlanetaryBodies |= true;
+		bRebuildSequence = true;
 	}
-
-	if (BuildAll || BuildStars)
+	
+	if (bBuildAll || bBuildStars)
 	{
 		InitStars();
 	}
 
-	if (BuildAll || BuildPlanetaryBodies)
+	if (bBuildAll || bBuildPlanetaryBodies)
 	{
 		InitPlanetaryBodies();
 	}
 
-	if (RebuildSequence)
+	if (bRebuildSequence)
 	{
 		RootSequence = nullptr;
 		UpdateRootSequence();
@@ -503,7 +558,7 @@ void ACelestialVaultDaySequenceActor::InitStars()
 						if (CelestialStarInputData.Magnitude <= MaxVisibleMagnitude)
 						{
 							// Location - Convert to UE Left-handed Frame (Invert the Y coordinate)
-							FVector StarLocation = UCelestialMaths::RADECToXYZ_RH(CelestialStarInputData.RA * 15, CelestialStarInputData.DEC, CelestialVaultDistance * 1000.0 * StarsVaultPercentage / 100.0) *  FVector(100, -100, 100);
+							FVector StarLocation = UCelestialMaths::RADECToXYZ_RH(CelestialStarInputData.RA * 15.0, CelestialStarInputData.DEC, CelestialVaultDistance * 1000.0 * StarsVaultPercentage / 100.0) *  FVector(100.0, -100.0, 100.0);
 
 							// Color
 							FLinearColor StarColor = UCelestialMaths::BVtoLinearColor(CelestialStarInputData.ColorIndex);
@@ -518,7 +573,7 @@ void ACelestialVaultDaySequenceActor::InitStars()
 							StarsComponent->SetCustomData(NewIndex, NewCustomData);
 
 							// Keep trace of the Star Information for further runtime queries					
-							if (KeepStarsInfo)
+							if (bKeepStarsInfo)
 							{
 								FStarInfo StarInfo;
 								StarInfo.RA = CelestialStarInputData.RA;
@@ -556,7 +611,7 @@ void ACelestialVaultDaySequenceActor::InitStars()
 						if (StarInputData.Magnitude <= MaxVisibleMagnitude)
 						{
 							// Location - Convert to UE Left-handed Frame (Invert the Y coordinate)
-							FVector StarLocation = UCelestialMaths::RADECToXYZ_RH(StarInputData.RA * 15, StarInputData.DEC, CelestialVaultDistance * 1000.0 * StarsVaultPercentage / 100.0) * FVector(100, -100, 100);
+							FVector StarLocation = UCelestialMaths::RADECToXYZ_RH(StarInputData.RA * 15.0, StarInputData.DEC, CelestialVaultDistance * 1000.0 * StarsVaultPercentage / 100.0) * FVector(100.0, -100.0, 100.0);
 
 							// Color
 							FLinearColor StarColor = StarInputData.Color;
@@ -571,7 +626,7 @@ void ACelestialVaultDaySequenceActor::InitStars()
 							StarsComponent->SetCustomData(NewIndex, NewCustomData);
 
 							// Keep trace of the Star Information for further runtime queries					
-							if (KeepStarsInfo)
+							if (bKeepStarsInfo)
 							{
 								FStarInfo StarInfo;
 								StarInfo.RA = StarInputData.RA;
@@ -612,6 +667,7 @@ void ACelestialVaultDaySequenceActor::InitPlanetaryBodies()
 		// We don't want to use the GetMoonTransform function here because we also want to keep the Moon information for the day.
 		// It has a pitfall because the Moon information will not be adjusted to the current TOD, but that doesn't prevent anyone for doing a manuel query at some point
 		MoonBodyInfo = GetMoonInfo(JulianDay);
+		SetMoonDiscAge(MoonBodyInfo.Age);
 		MoonDiscComponent->SetRelativeTransform(MoonBodyInfo.UETransform);
 	}
 	
@@ -639,7 +695,7 @@ void ACelestialVaultDaySequenceActor::InitPlanetaryBodies()
 		// Init from the DataTable
 		PlanetsCatalog->ForeachRow<FPlanetaryBodyInputData>("AEarthDaySequenceActor::InitPlanetaryBodies", [&](const FName& Key, const FPlanetaryBodyInputData& InputPlanetaryBody)
 			{
-				double UEBodyDistance = CelestialVaultDistance * 100000 * PlanetsVaultPercentage / 100.0;
+				double UEBodyDistance = CelestialVaultDistance * 100000.0 * PlanetsVaultPercentage / 100.0;
 				FPlanetaryBodyInfo BodyInfo = GetPlanetaryBodyInfo(InputPlanetaryBody, JulianDay, UEBodyDistance, PlanetsScale);
 
 				// Add the new ISM Instance
@@ -653,7 +709,7 @@ void ACelestialVaultDaySequenceActor::InitPlanetaryBodies()
 				PlanetsComponent->SetCustomData(NewIndex, NewCustomData);
 
 				// Keep trace of the PlanetaryBody for further queries (the Datatable is readonly, so store in another object)
-				if (KeepPlanetsInfos)
+				if (bKeepPlanetsInfos)
 				{
 					PlanetsInfos.Add(BodyInfo);	
 				}
@@ -680,31 +736,15 @@ FPlanetaryBodyInfo ACelestialVaultDaySequenceActor::GetPlanetaryBodyInfo(const F
 	// Compute Magnitude
 	double Phase; // TODO - Check the PhaseAngle Computations. 
 	BodyInfo.ApparentMagnitude = UCelestialMaths::GetPlanetaryBodyMagnitude(InputPlanetaryBody, DistanceToEarthAU, DistanceToSunAU, DistanceEarthToSunAU, Phase);
-	BodyInfo.Phase = Phase;
+	BodyInfo.Age = Phase;
 	
 	// Compute the True and the Scaled apparent Diameters
-	BodyInfo.ApparentDiameterDegrees = FMath::RadiansToDegrees(FMath::Atan2(InputPlanetaryBody.Radius * 1000.0, UCelestialMaths::AstronomicalUnitsToMeters(DistanceToEarthAU)) * 2);
-	BodyInfo.ScaledApparentDiameterDegrees = FMath::RadiansToDegrees(FMath::Atan2(InputPlanetaryBody.Radius * 1000.0 * BodyScale, UCelestialMaths::AstronomicalUnitsToMeters(DistanceToEarthAU)) * 2);
-
-	// Location
-	FVector BodyLocation = UCelestialMaths::RADECToXYZ_RH(BodyInfo.RA * 15, BodyInfo.DEC, UEDistance);
-	BodyLocation *= FVector(1, -1, 1); // Convert to UE Frame by inverting Y
-
-	// Rotation - Orient the Body mesh towards the earth -
-	// Use Look at the UE Origin, this is an acceptable approximation if the vault is big enough - We could use the Earth center if needed
-	FVector BodyToEarth = (FVector::ZeroVector - BodyLocation).GetSafeNormal();
-	FQuat BodyRotation = FQuat(BodyToEarth.Rotation());
-
-	// Scale
-	// The mesh plane is 100 UE Units (1m), and located at PlanetMeshDistanceKm. Use Thales theorem to compute its effective scale at this distance to ensure the right apparent diameter
-	double UERadius = (UEDistance * InputPlanetaryBody.Radius * 1000.0) / UCelestialMaths::AstronomicalUnitsToMeters(DistanceToEarthAU);
-	double Scale = UERadius * BodyScale / 50.0; // The plane half-length is 50 UE Units
-	FVector BodyScale3D = FVector(Scale, Scale, Scale);
+	BodyInfo.ApparentDiameterDegrees = FMath::RadiansToDegrees(FMath::Atan2(InputPlanetaryBody.Radius * 1000.0, UCelestialMaths::AstronomicalUnitsToMeters(DistanceToEarthAU)) * 2.0);
+	BodyInfo.ScaledApparentDiameterDegrees = FMath::RadiansToDegrees(FMath::Atan2(InputPlanetaryBody.Radius * 1000.0 * BodyScale, UCelestialMaths::AstronomicalUnitsToMeters(DistanceToEarthAU)) * 2.0);
 
-	BodyInfo.UETransform = FTransform(BodyRotation, BodyLocation, BodyScale3D);
 
-	// Direction
-	BodyInfo.DirectionTowardEarth = BodyToEarth;
+	//
+	BodyInfo.ComputeTranform(UEDistance, BodyScale);
 	
 	return BodyInfo;
 }
@@ -743,6 +783,6 @@ void ACelestialVaultDaySequenceActor::RebuildAll()
 		InitPlanetaryBodies();
 
 		double CelestialVaultAngle = GetDayCelestialVaultAngle();
-		CelestialVaultComponent->SetRelativeRotation( FRotator(0, CelestialVaultAngle, 0.0) );
+		CelestialVaultComponent->SetRelativeRotation( FRotator(0.0, CelestialVaultAngle, 0.0) );
 	}
 }
diff --git a/Engine/Plugins/Experimental/CelestialVault/Source/CelestialVault/Private/CelestialVaultSequence.cpp b/Engine/Plugins/Experimental/CelestialVault/Source/CelestialVault/Private/CelestialVaultSequence.cpp
--- a/Engine/Plugins/Experimental/CelestialVault/Source/CelestialVault/Private/CelestialVaultSequence.cpp
+++ b/Engine/Plugins/Experimental/CelestialVault/Source/CelestialVault/Private/CelestialVaultSequence.cpp
@@ -37,14 +37,14 @@ void FCelestialVaultSequence::BuildSequence(UProceduralDaySequenceBuilder* InBui
 			double CelestialVaultAngleAtMidnight = EarthDaySequenceActor->GetDayCelestialVaultAngle(); 
 
 			const double KeyTime = Key * NormalizedTimeIncrement;
-			InBuilder->AddRotationKey(KeyTime, FRotator(0, CelestialVaultAngleAtMidnight + KeyTime * 360, 0.0));
+			InBuilder->AddRotationKey(KeyTime, FRotator(0.0, CelestialVaultAngleAtMidnight + KeyTime * 360.0, 0.0));
 		}
 	}
 
 	// Create the Moon Disc Animation
 	UStaticMeshComponent* MoonDiscComponent = EarthDaySequenceActor->MoonDiscComponent;
 	UDirectionalLightComponent* MoonLightComponent = EarthDaySequenceActor->MoonLightComponent;
-	
+
 	if (MoonDiscComponent && MoonLightComponent)
 	{
 		const double NormalizedTimeIncrement = 1.0 / FMath::Max(KeyCount - 1, static_cast<unsigned>(1));
@@ -52,12 +52,19 @@ void FCelestialVaultSequence::BuildSequence(UProceduralDaySequenceBuilder* InBui
 		{
 			const double KeyTime = Key * NormalizedTimeIncrement;
 
-			FDateTime KeyDateTime = LocalTimeMidnight + FTimespan::FromSeconds(KeyTime * 24*60*60);
-			FDateTime UTCTime = UCelestialMaths::LocalTimeToUTCTime(KeyDateTime, EarthDaySequenceActor->GMT_TimeZone, EarthDaySequenceActor->bIsDaylightSavings);
+			FDateTime KeyDateTime = LocalTimeMidnight + FTimespan::FromSeconds(KeyTime * 24.0 * 60.0 * 60.0);
+			FDateTime UTCTime = UCelestialMaths::LocalTimeToUTCTime(KeyDateTime, EarthDaySequenceActor->GMT_TimeZone,
+				EarthDaySequenceActor->bIsDaylightSavings);
 			double JulianDay = UCelestialMaths::UTCDateTimeToJulianDate(UTCTime);
 
 			FPlanetaryBodyInfo MoonInfo = EarthDaySequenceActor->GetMoonInfo(JulianDay);
 
+			if (Key == 0)
+			{
+				EarthDaySequenceActor->SetMoonDiscAge(MoonInfo.Age);
+			}
+
+
 			// Animate Disc
 			InBuilder->SetActiveBoundObject(MoonDiscComponent);
 			InBuilder->AddTransformKey(KeyTime, MoonInfo.UETransform);
@@ -65,7 +72,6 @@ void FCelestialVaultSequence::BuildSequence(UProceduralDaySequenceBuilder* InBui
 			// Animate Light
 			InBuilder->SetActiveBoundObject(MoonLightComponent);
 			InBuilder->AddRotationKey(KeyTime, FRotationMatrix::MakeFromX(MoonInfo.DirectionTowardEarth).Rotator());
-			
 		}
 	}
 
@@ -73,14 +79,15 @@ void FCelestialVaultSequence::BuildSequence(UProceduralDaySequenceBuilder* InBui
 	if (UDirectionalLightComponent* SunLightComponent = EarthDaySequenceActor->SunLightComponent)
 	{
 		InBuilder->SetActiveBoundObject(SunLightComponent);
-		
+
 		const double NormalizedTimeIncrement = 1.0 / FMath::Max(KeyCount - 1, static_cast<unsigned>(1));
 		for (unsigned int Key = 0; Key < KeyCount; ++Key)
 		{
 			const double KeyTime = Key * NormalizedTimeIncrement;
 
-			FDateTime KeyDateTime = LocalTimeMidnight + FTimespan::FromSeconds(KeyTime * 24*60*60);
-			FDateTime UTCTime = UCelestialMaths::LocalTimeToUTCTime(KeyDateTime, EarthDaySequenceActor->GMT_TimeZone, EarthDaySequenceActor->bIsDaylightSavings);
+			FDateTime KeyDateTime = LocalTimeMidnight + FTimespan::FromSeconds(KeyTime * 24 * 60 * 60);
+			FDateTime UTCTime = UCelestialMaths::LocalTimeToUTCTime(KeyDateTime, EarthDaySequenceActor->GMT_TimeZone,
+				EarthDaySequenceActor->bIsDaylightSavings);
 			double JulianDay = UCelestialMaths::UTCDateTimeToJulianDate(UTCTime);
 
 			FSunInfo SunInfo = EarthDaySequenceActor->GetSunInfo(JulianDay);
diff --git a/Engine/Plugins/Experimental/CelestialVault/Source/CelestialVault/Public/CelestialDataTypes.h b/Engine/Plugins/Experimental/CelestialVault/Source/CelestialVault/Public/CelestialDataTypes.h
--- a/Engine/Plugins/Experimental/CelestialVault/Source/CelestialVault/Public/CelestialDataTypes.h
+++ b/Engine/Plugins/Experimental/CelestialVault/Source/CelestialVault/Public/CelestialDataTypes.h
@@ -112,9 +112,13 @@ public:
 	UPROPERTY(BlueprintReadWrite, EditAnywhere, Category= "Celestial Vault")
 	double ApparentMagnitude = 0.0;
 
-	/** indication of the crescent phase. 0 = New, 0.5 = Half, 1 = Full*/ 
+	/** indication of the lunar age. 0 = New Moon, 0.25 = First quarter, 0.5 = Full Moon, 1 = Next New Moon */ 
 	UPROPERTY(BlueprintReadWrite, EditAnywhere, Category= "Celestial Vault")
-	double Phase = 0.0;
+	double Age = 0.5;
+
+	/** indication of percentage of a full moon illumination */ 
+	UPROPERTY(BlueprintReadWrite, EditAnywhere, Category= "Celestial Vault")
+	double IlluminationPercentage = 1.0;
 
 	/** Internal value of the ISM Instance corresponding to this planet */
 	UPROPERTY(BlueprintReadWrite, EditAnywhere, Category= "Celestial Vault")
@@ -130,6 +134,9 @@ public:
 
 	// Can't use a UFUNCTION inside a struct --> Use the Celestial Maths Blueprint Function Library
 	FString ToString() const;
+
+	
+	void ComputeTranform(double UEDistance, double BodyScale);
 };
 
 /**
@@ -141,11 +148,11 @@ struct FSunInfo
 	GENERATED_BODY()
 public:
 	/** The Sun Right Ascension in the Celestial Frame - In hours! */
-	UPROPERTY(BlueprintReadWrite, EditAnywhere, Category= "Celestial Vault", meta=(DisplayName="Right Ascension", MakeStructureDefaultValue="0.000000"))
+	UPROPERTY(BlueprintReadWrite, EditAnywhere, Category= "Celestial Vault", meta=(ForceUnits="Hours", DisplayName="Right Ascension", MakeStructureDefaultValue="0.000000"))
 	double RA = 0.0;
 
 	/** The Sun Declination in the Celestial Frame - In Degrees */
-	UPROPERTY(BlueprintReadWrite, EditAnywhere, Category= "Celestial Vault", meta=(DisplayName="Declination", MakeStructureDefaultValue="0.000000"))
+	UPROPERTY(BlueprintReadWrite, EditAnywhere, Category= "Celestial Vault", meta=(ForceUnits="Degrees", DisplayName="Declination", MakeStructureDefaultValue="0.000000"))
 	double DEC = 0.0;
 
 	/** Keep track of the Vault-relative transform, at ToD = 0 for the Sun */
@@ -179,11 +186,11 @@ struct FStarInfo
 
 public:
 	/** The Star Right Ascension in the Celestial Frame - In hours! */
-	UPROPERTY(BlueprintReadWrite, EditAnywhere, Category= "Celestial Vault", meta=(DisplayName="Right Ascension", MakeStructureDefaultValue="0.000000"))
+	UPROPERTY(BlueprintReadWrite, EditAnywhere, Category= "Celestial Vault", meta=(ForceUnits="Hours", DisplayName="Right Ascension", MakeStructureDefaultValue="0.000000"))
 	double RA = 0.0;
 
 	/** The Star Declination in the Celestial Frame - In Degrees */
-	UPROPERTY(BlueprintReadWrite, EditAnywhere, Category= "Celestial Vault", meta=(DisplayName="Declination", MakeStructureDefaultValue="0.000000"))
+	UPROPERTY(BlueprintReadWrite, EditAnywhere, Category= "Celestial Vault", meta=(ForceUnits="Degrees", DisplayName="Declination", MakeStructureDefaultValue="0.000000"))
 	double DEC = 0.0;
 
 	/** Earth to Star distance (in Parsecs)  */
@@ -238,11 +245,11 @@ struct FStarInputData : public FTableRowBase
 	GENERATED_BODY()
 public:
 	/** The Star Right Ascension in the Celestial Frame - In hours! */
-	UPROPERTY(BlueprintReadWrite, EditAnywhere, Category= "Celestial Vault", meta=(DisplayName="Right Ascension", MakeStructureDefaultValue="0.000000"))
+	UPROPERTY(BlueprintReadWrite, EditAnywhere, Category= "Celestial Vault", meta=(ForceUnits="Hours", DisplayName="Right Ascension", MakeStructureDefaultValue="0.000000"))
 	double RA = 0.0;
 
 	/** The Star Declination in the Celestial Frame - In Degrees */
-	UPROPERTY(BlueprintReadWrite, EditAnywhere, Category= "Celestial Vault", meta=(DisplayName="Declination", MakeStructureDefaultValue="0.000000"))
+	UPROPERTY(BlueprintReadWrite, EditAnywhere, Category= "Celestial Vault", meta=(ForceUnits="Degrees", DisplayName="Declination", MakeStructureDefaultValue="0.000000"))
 	double DEC = 0.0;
 
 	/** Earth to Star distance (in Parsecs)  */
diff --git a/Engine/Plugins/Experimental/CelestialVault/Source/CelestialVault/Public/CelestialMaths.h b/Engine/Plugins/Experimental/CelestialVault/Source/CelestialVault/Public/CelestialMaths.h
--- a/Engine/Plugins/Experimental/CelestialVault/Source/CelestialVault/Public/CelestialMaths.h
+++ b/Engine/Plugins/Experimental/CelestialVault/Source/CelestialVault/Public/CelestialMaths.h
@@ -79,6 +79,18 @@ public:
     UFUNCTION(BlueprintCallable, Category="Celestial Vault|Planetary bodies")
     static double GetPlanetaryBodyMagnitude(FPlanetaryBodyInputData PlanetaryBody, double DistanceToSunAU, double DistanceToEarthAU, double DistanceEarthToSunAU, double& PhaseAngle);
 
+    /** Returns the Moon Phase for a specific Date
+     * 
+     *  This is an approximate computation using a number of lunar cycles with a synodic month equals to 29.53059 days
+     *  Not very precise over more than 1 centuries before of after 2025
+     */
+    UFUNCTION(BlueprintCallable, Category="Celestial Vault|Planetary bodies")
+    static double GetMoonNormalizedAgeSimple(double JulianDate);
+
+	/** return the illumination factor (0..1) of a Body, considering his normalized age and the crescent effects */
+	UFUNCTION(BlueprintCallable, Category="Celestial Vault|Planetary bodies")
+	static double GetIlluminationPercentage(double NormalizedAge);
+
 #pragma endregion
 
 #pragma region Time
@@ -315,6 +327,8 @@ public:
     static FString Conv_DMSToString(bool Sign, int32 Degrees, int32 Minutes, double Seconds);
     
 #pragma endregion
+public:
+    static double SynodicMonthAverage;
     
 private:
 
@@ -329,4 +343,6 @@ private:
     static FMatrix VSOPToJ2000;
     static double SpeedOfLightMetersPerSeconds;
     static double AstronomicalUnitsMeters;
+    static double NewMoonReferenceJulianDate;
+    
 };
diff --git a/Engine/Plugins/Experimental/CelestialVault/Source/CelestialVault/Public/CelestialVaultDaySequenceActor.h b/Engine/Plugins/Experimental/CelestialVault/Source/CelestialVault/Public/CelestialVaultDaySequenceActor.h
--- a/Engine/Plugins/Experimental/CelestialVault/Source/CelestialVault/Public/CelestialVaultDaySequenceActor.h
+++ b/Engine/Plugins/Experimental/CelestialVault/Source/CelestialVault/Public/CelestialVaultDaySequenceActor.h
@@ -152,7 +152,7 @@ public: // Properties
 
 	/** If true, the Stars information will be kept in memory and queryable at runtime */ 
 	UPROPERTY(EditAnywhere, BlueprintReadWrite, Category="Celestial Vault|Stars")
-	bool KeepStarsInfo = false;
+	bool bKeepStarsInfo = false;
 	
 	/** Array of the created Stars information - Only populated if KeepStarsInfo is true */ 
 	UPROPERTY(BlueprintReadOnly, Category="Celestial Vault|Stars")
@@ -172,7 +172,7 @@ public: // Properties
 
 	/** If true, the Stars information will be kept in memory and queryable at runtime */ 
 	UPROPERTY(EditAnywhere, BlueprintReadWrite, Category="Celestial Vault|Planets")
-	bool KeepPlanetsInfos = false;
+	bool bKeepPlanetsInfos = false;
 	
 	/** Array of the created planetary bodies, with all their computed information - Only populated if KeepPlanetsInfos is true */ 
 	UPROPERTY(BlueprintReadOnly, Category="Celestial Vault|Planets")
@@ -190,10 +190,47 @@ public: // Properties
 	UPROPERTY(BlueprintReadOnly, Category="Celestial Vault|Moons")
 	FPlanetaryBodyInfo MoonBodyInfo;
 
+	/** If true, the moon Age (Phase) and location can be overriden */  
+	UPROPERTY(EditAnywhere, Category="Celestial Vault|Moons")
+	bool bManualControl = false;
+	
+	/** Lunar age. 0 = New Moon, 0.25 = First quarter, 0.5 = Full Moon, 1 = Next New Mo */  
+	UPROPERTY(EditAnywhere, Category="Celestial Vault|Moons", meta = (EditConditionHides, EditCondition = "bManualControl==true", ClampMin = "0", ClampMax = "1"))
+	float MoonAge = 0.1f;
+	
+	/** When faking the moon location, we need to give a location relative to the sun
+	 * This is a way to control this "Horizontally" using an offset in Right Ascension. 
+	 * 
+	 * Ex: if MoonOffset_RA is 4, the moon will set 4 hours after the sun, so will be visible only the 4 first hours of the night.
+	 * Be mindful, it should normally be correlated with the phase.
+	 *   (if HoursBehindSun < 12, we are ~1st quarter, so Age should be < 0.5, and if 12< < 24 age should be > 0.5) 
+	 * Any combination is possible, but manual control can lead to inconsistencies. 
+	 */  
+	UPROPERTY(EditAnywhere, Category="Celestial Vault|Moons", meta = (EditConditionHides, DisplayName="Moon Offset to Sun's Right Ascension", EditCondition = "bManualControl==true", ForceUnits="Hours", ClampMin = "0", ClampMax = "24"))
+	float MoonOffset_RA = 6.0f;
+
+	/** When faking the moon location, we need to give a location relative to the sun
+	 * This is a way to control this "Vertically" using an offset in Declination. */ 
+	UPROPERTY(EditAnywhere, Category="Celestial Vault|Moons", meta = (EditConditionHides, DisplayName="Moon Offset to Sun's Declination", EditCondition = "bManualControl==true", ForceUnits="Degrees", ClampMin = "-45", ClampMax = "45"))
+	float MoonOffset_DEC = 15.0f;
+
 	/** Celestial Info for the Sun, at the beginning of the Day */  
 	UPROPERTY(BlueprintReadOnly, Category="Celestial Vault|Sun")
 	FSunInfo SunInfo; 
 
+	/** Base Sun Intensity
+	 * Typically 120000 Lux
+	 */
+	UPROPERTY(EditAnywhere, BlueprintReadWrite, AdvancedDisplay, Category="Celestial Vault", meta=( ForceUnits="Lux", ClampMin = "0", ClampMax = "200000"))
+	float SunLightIntensity = 120000.0f;
+	
+	/** Base Moonlight Intensity (for Full Moon)
+	 * Typically 0.1 Lux, up to 0.32 Lux when the moon is at it's perigee (SuperMoon)
+	 */
+	UPROPERTY(EditAnywhere, BlueprintReadWrite, AdvancedDisplay, Category="Celestial Vault", meta=( ForceUnits="Lux", ClampMin = "0", ClampMax = "200000"))
+	float MoonLightIntensity = 0.1f;
+
+	
 #pragma endregion 
 
 protected: // Functions
@@ -216,6 +253,10 @@ public: // Functions
 	UFUNCTION(BlueprintCallable, Category = "Celestial Vault|Queries")
 	FPlanetaryBodyInfo GetMoonInfo(double JulianDate);
 
+	/** Manually set the Moon Age (Phase) */  
+	UFUNCTION(BlueprintCallable, Category = "Celestial Vault|Moons")
+	void SetMoonDiscAge(float InMoonAge);
+	
 // Data Queries
 	
 	/** Return the Celestial Information of the Star closest to a specific direction, within an angle threshold */
@@ -234,6 +275,7 @@ public: // Functions
 	double GetDayCelestialVaultAngle() const;
 
 #if WITH_EDITOR
+	
 	virtual void PostEditChangeProperty(FPropertyChangedEvent& PropertyChangedEvent) override;
 #endif
 
