- Added ExponentialHeightfog and GlobalPostProcess Volume into the CelestialVaultDaySequence Actor with correct exposure presets.
- Ensured all assets have a proper naming convention
- Improved the visual quality of all shaders, with a proper reaction to ToD and SkyAtmosphere
(Additive Luminance, Cloud Fogging, DoF support, Moon Color variation, Size/luminance presets, Right Luminance from Magniture, Artistic controls)
- Added a Credits.txt file

#virtualized

diff --git a/Engine/Plugins/Experimental/CelestialVault/Resources/Credits.txt b/Engine/Plugins/Experimental/CelestialVault/Resources/Credits.txt
new file mode 100644
--- /dev/null
+++ b/Engine/Plugins/Experimental/CelestialVault/Resources/Credits.txt
@@ -0,0 +1,72 @@
+****** Celestial Vault Plugin ******
+*             Credits              *
+************************************
+
+##############
+### Assets ###
+##############
+
+-------------------------------------------------------------------------
+-- Milky Way, Constellation Lines, RADEC Grid in Celestial coordinates -- 
+-------------------------------------------------------------------------
+
+Credits: NASA/Goddard Space Flight Center Scientific Visualization Studio. Gaia DR2: ESA/Gaia/DPAC. Constellation figures based on those developed for the IAU by Alan MacRobert of Sky and Telescope magazine (Roger Sinnott and Rick Fienberg).
+
+Note - You can download a higher resolution of these textures at https://svs.gsfc.nasa.gov/4851/
+
+-------------------
+-- Moon Textures --
+-------------------
+
+Credits: NASA's Scientific Visualization Studio 
+
+Note - You can download a higher resolution of these textures at https://svs.gsfc.nasa.gov/4720/ 
+(The Normal Map has been created from the Displacement map)
+
+-----------------------
+-- Fictional Planets --
+-----------------------
+Various sources from public domain or Creative Commons Attribution 4.0 International
+Provided as an implementation example. Please do use your own in production. 
+
+https://opengameart.org/content/planets-6x3000x3000
+https://commons.wikimedia.org/wiki/File:Fictional_Planet_Andoria.png
+https://commons.wikimedia.org/wiki/File:Fictional_Planet_Qo%27noS.png
+
+--------------------
+-- Stars Catalogs --
+--------------------
+
+HYG Database : Current version (2024), containing all stars in Hipparcos, Yale Bright Star, and Gliese catalogs (almost 120,000 stars)
+
+The HYG database is licensed with the Creative Commons Attribution-ShareAlike 4.0 license.
+https://www.astronexus.com/projects/hyg
+
+There have been filtered and processed to a CSV file imported as a DataTable. One containing the full catalog, another limited to the first 10K brightests (~Mag 6)
+
+
+##############################
+### Celestial Computations ###
+##############################
+
+Big shoutout to Greg's Miller for his awesome Astronomy Programming Page 
+https://www.celestialprogramming.com/
+
+
+VSOP Computations have been adapted from [REDACTED_PROFILE_URL]/vsop87-multilang
+License : 
+Author: [REDACTED_PERSON] ([REDACTED_EMAIL])
+I release all of my work in this project to the public domain.
+Note, though, not everything here can be considered my work.  The VSOP87 theory is the work of Bureau des Longitudes. https://site.bdlg.fr/
+
+
+Other algorithms have been implementated following the equations described in 
+- Jean Meeus : Astronomical Algorithms Second Edition
+- James L. Hilton : Computing Apparent Planetary Magnitudes for The Astronomical Almanac
+- George H. Kaplan : The IAU Resolutions on Astronomical Reference Systems, Time Scales, and Earth Rotation Models
+- https://en.wikipedia.org/wiki/Color_index
+- https://ftp.imcce.fr/pub/ephem/planets/vsop87/vsop87.doc
+- https://aa.usno.navy.mil/faq/GAST
+- http://maia.usno.navy.mil/ser7/tai-utc.dat
+- https://www2.mps.mpg.de/homes/fraenz/systems/systems2art/node2.html
+
diff --git a/Engine/Plugins/Experimental/CelestialVault/Source/CelestialVault/Private/CelestialVaultDaySequenceActor.cpp b/Engine/Plugins/Experimental/CelestialVault/Source/CelestialVault/Private/CelestialVaultDaySequenceActor.cpp
--- a/Engine/Plugins/Experimental/CelestialVault/Source/CelestialVault/Private/CelestialVaultDaySequenceActor.cpp
+++ b/Engine/Plugins/Experimental/CelestialVault/Source/CelestialVault/Private/CelestialVaultDaySequenceActor.cpp
@@ -7,6 +7,7 @@
 // Components
 #include "Components/DirectionalLightComponent.h"
 #include "Components/ExponentialHeightFogComponent.h"
+#include "Components/PostProcessComponent.h"
 #include "Components/InstancedStaticMeshComponent.h"
 #include "Components/SkyAtmosphereComponent.h"
 #include "Components/SkyLightComponent.h"
@@ -21,6 +22,7 @@
 #include "DaySequenceCollectionAsset.h"
 #include "Engine/World.h"
 #include "Engine/StaticMesh.h"
+#include "Curves/CurveFloat.h"
 
 // Celestial Objects
 #include "CelestialMaths.h"
@@ -33,10 +35,24 @@
 ACelestialVaultDaySequenceActor::ACelestialVaultDaySequenceActor(const FObjectInitializer& Init)
 : Super(Init)
 {
-	// ExponentialHeightFogComponent = CreateOptionalDefaultSubobject<UExponentialHeightFogComponent>(TEXT("ExponentialHeightFog"));
-	// ExponentialHeightFogComponent->SetupAttachment(RootComponent);
-	// ExponentialHeightFogComponent->bEnableVolumetricFog = true;
-
+	ExponentialHeightFogComponent = CreateOptionalDefaultSubobject<UExponentialHeightFogComponent>(TEXT("ExponentialHeightFog"));
+	ExponentialHeightFogComponent->SetupAttachment(RootComponent);
+	ExponentialHeightFogComponent->bEnableVolumetricFog = true;
+
+	static ConstructorHelpers::FObjectFinder<UCurveFloat> HighlightContrastCurve(TEXT("/CelestialVault/Data/CF_CelestialHighlightContrastCurve.CF_CelestialHighlightContrastCurve"));
+	GlobalPostProcessVolume = CreateOptionalDefaultSubobject<UPostProcessComponent>(TEXT("GlobalPostProcessVolume"));
+	GlobalPostProcessVolume->SetupAttachment(RootComponent);
+	GlobalPostProcessVolume->Settings.bOverride_AutoExposureMinBrightness = true;
+	GlobalPostProcessVolume->Settings.AutoExposureMinBrightness = -0.5f;
+	GlobalPostProcessVolume->Settings.bOverride_LocalExposureHighlightContrastCurve = true;
+	GlobalPostProcessVolume->Settings.LocalExposureHighlightContrastCurve = HighlightContrastCurve.Object;
+	GlobalPostProcessVolume->Settings.bOverride_LocalExposureDetailStrength = true;
+	GlobalPostProcessVolume->Settings.LocalExposureDetailStrength = 1.2f;
+	GlobalPostProcessVolume->Settings.bOverride_LocalExposureBlurredLuminanceBlend = true;
+	GlobalPostProcessVolume->Settings.LocalExposureBlurredLuminanceBlend = 0.4f;
+	
+	
+	
 	// Components Attached to Root
 	SkyLightComponent = CreateDefaultSubobject<USkyLightComponent>(TEXT("Sky Light"));
 	SkyLightComponent->SetupAttachment(RootComponent);
@@ -68,7 +84,7 @@ ACelestialVaultDaySequenceActor::ACelestialVaultDaySequenceActor(const FObjectIn
 	// Components attached to the Celestial Vaul
 	
 	// Deep Sky background
-	static ConstructorHelpers::FObjectFinder<UStaticMesh> SkySphereDefaultMesh(TEXT("/CelestialVault/Meshes/CelestialVault.CelestialVault"));
+	static ConstructorHelpers::FObjectFinder<UStaticMesh> SkySphereDefaultMesh(TEXT("/CelestialVault/Meshes/SM_CelestialVault.SM_CelestialVault"));
 	static ConstructorHelpers::FObjectFinder<UMaterialInterface> SkySphereDefaultMaterial(TEXT("/CelestialVault/Materials/MI_CelestialVault.MI_CelestialVault"));
 	DeepSkyComponent = CreateOptionalDefaultSubobject<UStaticMeshComponent>(TEXT("Deep Sky"));
 	DeepSkyComponent->SetupAttachment(CelestialVaultComponent);
@@ -82,7 +98,7 @@ ACelestialVaultDaySequenceActor::ACelestialVaultDaySequenceActor(const FObjectIn
 	DeepSkyComponent->SetRelativeScale3D(FVector(CelestialVaultDistance*1000.0));
 
 	// Stars ISM
-	static ConstructorHelpers::FObjectFinder<UStaticMesh> PlaneXMesh(TEXT("/CelestialVault/Meshes/PlaneX.PlaneX"));
+	static ConstructorHelpers::FObjectFinder<UStaticMesh> PlaneXMesh(TEXT("/CelestialVault/Meshes/SM_Plane_FacingX.SM_Plane_FacingX"));
 	static ConstructorHelpers::FObjectFinder<UMaterialInterface> StarsDefaultMaterial(TEXT("/CelestialVault/Materials/MI_Stars.MI_Stars"));
 	StarsComponent = CreateOptionalDefaultSubobject<UInstancedStaticMeshComponent>(TEXT("Stars"));
 	StarsComponent->SetupAttachment(CelestialVaultComponent);
@@ -145,10 +161,10 @@ ACelestialVaultDaySequenceActor::ACelestialVaultDaySequenceActor(const FObjectIn
 		static ConstructorHelpers::FObjectFinder<UDaySequenceCollectionAsset> DefaultCollection(TEXT("/CelestialVault/DSCA_CelestialVault.DSCA_CelestialVault"));
 		DaySequenceCollections.Add(DefaultCollection.Object.Get());
 
-		static ConstructorHelpers::FObjectFinder<UDataTable> DefaultStarsCatalog(TEXT("/CelestialVault/Data/HYGCatalog_10K.HYGCatalog_10K"));
+		static ConstructorHelpers::FObjectFinder<UDataTable> DefaultStarsCatalog(TEXT("/CelestialVault/Data/DT_HYGCatalog_10K.DT_HYGCatalog_10K"));
 		CelestialStarCatalog = DefaultStarsCatalog.Object.Get();
 
-		static ConstructorHelpers::FObjectFinder<UDataTable> DefaultPlanetaryBodiesCatalog(TEXT("/CelestialVault/Data/SolarSystemPlanets.SolarSystemPlanets"));
+		static ConstructorHelpers::FObjectFinder<UDataTable> DefaultPlanetaryBodiesCatalog(TEXT("/CelestialVault/Data/DT_SolarSystemPlanets.DT_SolarSystemPlanets"));
 		PlanetsCatalog = DefaultPlanetaryBodiesCatalog.Object.Get();
 	}
 }
@@ -209,8 +225,9 @@ FDateTime ACelestialVaultDaySequenceActor::GetDate() const
 	if ( Day > FDateTime::DaysInMonth(Year, Month))
 	{
 		int32 ClampedDay = FDateTime::DaysInMonth(Year, Month);
+		int32 ClampedMonth = FMath::Clamp(Month, 1, 12);
 		UE_LOG(LogCelestialVault, Warning, TEXT("Day value (%d) over the number of days in month - Using %d instead"), Day, ClampedDay);
-		return FDateTime(Year, Month, ClampedDay, 0,0,0);
+		return FDateTime(Year, ClampedMonth, ClampedDay, 0,0,0);
 	}
 		
 	return FDateTime(Year, Month, Day, 0,0,0);
@@ -753,6 +770,9 @@ void ACelestialVaultDaySequenceActor::RebuildAll()
 {
 	if (const UWorld* World = GetWorld())
 	{
+		// Make sure the actor is properly located at the Origin
+		SetActorTransform(FTransform::Identity);
+		
 		// DateTime has changed, we need to update the Reference CelestialVault Angle.
 		FDateTime DateTime = GetDate();
 		GMST_AtTOD_0 = UCelestialMaths::DateTimeToGreenwichMeanSiderealTime(DateTime);
diff --git a/Engine/Plugins/Experimental/CelestialVault/Source/CelestialVault/Public/CelestialVaultDaySequenceActor.h b/Engine/Plugins/Experimental/CelestialVault/Source/CelestialVault/Public/CelestialVaultDaySequenceActor.h
--- a/Engine/Plugins/Experimental/CelestialVault/Source/CelestialVault/Public/CelestialVaultDaySequenceActor.h
+++ b/Engine/Plugins/Experimental/CelestialVault/Source/CelestialVault/Public/CelestialVaultDaySequenceActor.h
@@ -15,7 +15,8 @@ class USkyAtmosphereComponent;
 class USkyLightComponent;
 class UVolumetricCloudComponent;
 class UDirectionalLightComponent;
-//class UExponentialHeightFogComponent;
+class UExponentialHeightFogComponent;
+class UPostProcessComponent;
 class UStaticMeshComponent;
 class UInstancedStaticMeshComponent;
 
@@ -51,9 +52,11 @@ public: // Properties
 	UPROPERTY(VisibleDefaultsOnly, BlueprintReadOnly, Category = "Celestial Vault", meta = (AllowPrivateAccess = "true"))
 	TObjectPtr<USkyLightComponent> SkyLightComponent;
 	
-	//TODO - Should HeightFog be optional because it doesn't work round earth when far from the origin?
-	// UPROPERTY(VisibleDefaultsOnly, BlueprintReadOnly, Category = "Celestial Vault", meta = (AllowPrivateAccess = "true"))
-	// TObjectPtr<UExponentialHeightFogComponent> ExponentialHeightFogComponent;
+	UPROPERTY(VisibleDefaultsOnly, BlueprintReadOnly, Category = "Celestial Vault", meta = (AllowPrivateAccess = "true"))
+	TObjectPtr<UExponentialHeightFogComponent> ExponentialHeightFogComponent;
+
+	UPROPERTY(VisibleDefaultsOnly, BlueprintReadOnly, Category = "Celestial Vault", meta = (AllowPrivateAccess = "true"))
+	TObjectPtr<UPostProcessComponent> GlobalPostProcessVolume;
 	
 	UPROPERTY(VisibleDefaultsOnly, BlueprintReadOnly, Category = "Celestial Vault", meta = (AllowPrivateAccess = "true"))
 	TObjectPtr<UVolumetricCloudComponent> VolumetricCloudComponent;
@@ -184,7 +187,7 @@ public: // Properties
 
 	/** Factor to artificially increase the Moon size */
 	UPROPERTY(EditAnywhere, BlueprintReadWrite, Category="Celestial Vault|Moons", meta = (ClampMin = "0"))
-	float MoonScale = 1.0f;
+	float MoonScale = 2.0f;
 
 	/** Celestial Info for the Moon, at the beginning of the Day */  
 	UPROPERTY(BlueprintReadOnly, Category="Celestial Vault|Moons")
@@ -196,7 +199,7 @@ public: // Properties
 	
 	/** Lunar age. 0 = New Moon, 0.25 = First quarter, 0.5 = Full Moon, 1 = Next New Mo */  
 	UPROPERTY(EditAnywhere, Category="Celestial Vault|Moons", meta = (EditConditionHides, EditCondition = "bManualControl==true", ClampMin = "0", ClampMax = "1"))
-	float MoonAge = 0.1f;
+	float MoonAge = 0.2f;
 	
 	/** When faking the moon location, we need to give a location relative to the sun
 	 * This is a way to control this "Horizontally" using an offset in Right Ascension. 
@@ -207,7 +210,7 @@ public: // Properties
 	 * Any combination is possible, but manual control can lead to inconsistencies. 
 	 */  
 	UPROPERTY(EditAnywhere, Category="Celestial Vault|Moons", meta = (EditConditionHides, DisplayName="Moon Offset to Sun's Right Ascension", EditCondition = "bManualControl==true", ForceUnits="Hours", ClampMin = "0", ClampMax = "24"))
-	float MoonOffset_RA = 6.0f;
+	float MoonOffset_RA = 12.0f;
 
 	/** When faking the moon location, we need to give a location relative to the sun
 	 * This is a way to control this "Vertically" using an offset in Declination. */ 
