Fixed the computation of the sun light direction that was wrong, despite having the right celestial coordinates

diff --git a/Engine/Plugins/Experimental/CelestialVault/Source/CelestialVault/Private/CelestialVaultDaySequenceActor.cpp b/Engine/Plugins/Experimental/CelestialVault/Source/CelestialVault/Private/CelestialVaultDaySequenceActor.cpp
--- a/Engine/Plugins/Experimental/CelestialVault/Source/CelestialVault/Private/CelestialVaultDaySequenceActor.cpp
+++ b/Engine/Plugins/Experimental/CelestialVault/Source/CelestialVault/Private/CelestialVaultDaySequenceActor.cpp
@@ -239,7 +239,7 @@ FSunInfo ACelestialVaultDaySequenceActor::GetSunInfo(double JulianDate)
 
 	// Location
 	FVector SunLocation = UCelestialMaths::RADECToXYZ_RH(ResultSunInfo.RA * 15.0, ResultSunInfo.DEC, 1000.0); // TODO - Add the Sun Distance To SunInfo 
-	SunLocation.Y *= -100.0; // Convert to UE Frame by inverting Y and scaling to UE Units
+	SunLocation.Y *= -1.0; // Convert to UE Frame by inverting Y
 	ResultSunInfo.UETransform.SetLocation(SunLocation);
 	ResultSunInfo.DirectionTowardEarth = (FVector::ZeroVector - SunLocation).GetSafeNormal();
 	return ResultSunInfo;
