* Refactored the Celestial Data Types for better consistency
  - Separated in different files the Celestial Bodies properties from the DataTable properties used to build them
  - Removed any dependency on the InputData from the CelestialMaths class
  - Cleaned up the class hierarchies
  - Refactored the VSOP87 computations, using data tables instead of replicated code

* Added the CelestialVaultDev Test Project in Sandbox/Enterprise

* Graphical Changes
   - Improved the Star Mask Textures, added new ones

* Precision improvements
  - Using RA DEC of date instead of the previouly computed J2000 RA,DEC.
  - Now considering the Precession and Nutation effects for the computations
  - Minor fixes on the GMST computation
  - Fixed precision issue in the VSOP87 computations (Wrong Julian Time value)
  - Consider the Earth Velocity in the RA.
  - Added a function to compute the transformation to apply to a planet mesh to make it tangent to a lat/long (And removed an abiguity between UE and ECEF Frame)

* Code cleanup
  - Querying the closest Planet or Star is not dependent on the ISM Index, as it could change over time.
  - Added guardrails for divide by zero
  - Fixed some comments

#virtualized
#Synced-CL 49312238

diff --git a/Engine/Plugins/Experimental/CelestialVault/Source/CelestialVault/Private/CelestialDataTypes.cpp b/Engine/Plugins/Experimental/CelestialVault/Source/CelestialVault/Private/CelestialDataTypes.cpp
--- a/Engine/Plugins/Experimental/CelestialVault/Source/CelestialVault/Private/CelestialDataTypes.cpp
+++ b/Engine/Plugins/Experimental/CelestialVault/Source/CelestialVault/Private/CelestialDataTypes.cpp
@@ -1,54 +1,48 @@
 ﻿// Copyright Epic Games, Inc. All Rights Reserved.
 
-
 #include "CelestialDataTypes.h"
 #include "CelestialMaths.h"
+#include "CelestialVault.h"
 
-// Earth Preset
-FPlanetaryBodyInputData FPlanetaryBodyInputData::Earth(TEXT("Earth"), EOrbitType::Earth, 6378.0);
-// Moon Preset
-FPlanetaryBodyInputData FPlanetaryBodyInputData::Moon(TEXT("Moon"), EOrbitType::Moon, 1737.4);
-
-
-FString FPlanetaryBodyInfo::ToString() const
+FString FCelestialBody::ToString() const
 {
-	FString PlanetaryBodyOrbitTypeName = FString("Invalid");
-	if (UEnum* EnumClass = StaticEnum<EOrbitType>())
-	{
-		PlanetaryBodyOrbitTypeName = EnumClass->GetNameStringByValue(static_cast<int64>(OrbitType));
-	}
-
 	return FString::Printf(TEXT(
-			"Planetary Body: %s\n"
-			"Orbit Type: %s\n"
-			"RightAscensionHours: %s\n"
-			"Declination: %s\n"
-			"Distance: %.0f AU\n"
+			"Name: %s\n"
+			"RA J2000: %s\n"
+			"DEC J2000: %s\n"
+			"RA of Date: %s\n"
+			"DEC of Date: %s\n"
+			"Distance: %.0f AU (%.0f Parsecs)\n"
 			"Radius: %.0f km\n"
-			"Apparent Diameter: %s / %.8f°\n"
-			"Apparent Magnitude: %.2f\n"
-			"Age: %.2f (%.2f days) since last New Moon\n"
-			"IlluminationPercentage: %.2f\n"
-			"-------\n"
-			"Scaled Apparent Diameter: %.8f°\n"
-			),
-			*Name,
-			*PlanetaryBodyOrbitTypeName,
-			*UCelestialMaths::Conv_RightAscensionToString(RA),
-			*UCelestialMaths::Conv_DeclinationToString(DEC),
-			DistanceInAU,
-			Radius,
-			*UCelestialMaths::Conv_DeclinationToString(ApparentDiameterDegrees), ApparentDiameterDegrees,
-			ApparentMagnitude,
-			Age, UCelestialMaths::SynodicMonthAverage * Age, 
-			IlluminationPercentage,
-			ScaledApparentDiameterDegrees
-		); 
-	
+			"Magnitude: %.2f\n"),
+		*Name,
+		*UCelestialMaths::Conv_RightAscensionToString(RAJ2000),
+		*UCelestialMaths::Conv_DeclinationToString(DECJ2000),
+		*UCelestialMaths::Conv_RightAscensionToString(RA),
+		*UCelestialMaths::Conv_DeclinationToString(DEC),
+		DistanceInAU, UCelestialMaths::AstronomicalUnitsToParsec(DistanceInAU),
+		Radius,
+		Magnitude
+	);
 }
 
-void FPlanetaryBodyInfo::ComputeTranform(double UEDistance, double BodyScale)
+void FCelestialBody::ComputeUELocalTransform(double ExtraScale, double UEDistanceOverride)
 {
+	if (DistanceInAU <= 0)
+	{
+		UE_LOG(LogCelestialVault, Error, TEXT("Invalid DistanceInAU (%f <= 0) for Celestial Body %s - Check your inputs"), DistanceInAU, *Name);
+		UELocalTransform = FTransform::Identity;
+		DirectionTowardEarth = -FVector::XAxisVector;
+		return; 
+	}
+	
+	// What Distance to use
+	double UEDistance = UCelestialMaths::AstronomicalUnitsToMeters(DistanceInAU) * 100.0; // We want UE Units there
+	if (UEDistanceOverride > 0.0)
+	{
+		UEDistance = UEDistanceOverride;
+	}
+
 	// Location
 	FVector BodyLocation = UCelestialMaths::RADECToXYZ_RH(RA * 15, DEC, UEDistance);
 	BodyLocation.Y *= -1; // Convert to UE Frame by inverting Y
@@ -59,43 +53,93 @@ void FPlanetaryBodyInfo::ComputeTranform(double UEDistance, double BodyScale)
 	FQuat BodyRotation = FQuat(BodyToEarth.Rotation());
 
 	// Scale
-	// The mesh plane is 100 UE Units (1m), and located at PlanetMeshDistanceKm. Use Thales theorem to compute its effective scale at this distance to ensure the right apparent diameter
-	double UERadius = (UEDistance * Radius * 1000.0) / UCelestialMaths::AstronomicalUnitsToMeters(DistanceInAU);
-	double Scale = UERadius * BodyScale / 50.0; // The plane half-length is 50 UE Units
+	// The mesh plane is 100 UE Units (1m), and located at UEDistance. Use Thales theorem to compute its effective scale at this distance to ensure the right apparent diameter
+	//                                                                             Body
+	//                                                                             |
+	//                                    UEObject                                 |  Body Radius
+	//                                     |  UE Radius                            |
+	//   Earth ----------------------------|---------------------------------------/
+	//                                    UE Distance                             Body Distance
+	//
+	//    Body Radius (m)        UE Radius (UE units) 
+	// => ------------------ = -------------------------
+	//    Body Distance (m)      UE Distance (UE Units)
+	//
+	
+	double UERadius = (UEDistance * Radius * 1000.0) / UCelestialMaths::AstronomicalUnitsToMeters(DistanceInAU);  // * 1000 is to transform the Body Radius (km) into (m)
+	double Scale = UERadius / 50.0 * ExtraScale; // The plane half-length is 50 UE Units
 	FVector BodyScale3D = FVector(Scale, Scale, Scale);
 
-	UETransform = FTransform(BodyRotation, BodyLocation, BodyScale3D);
+	// Build Transform
+	UELocalTransform = FTransform(BodyRotation, BodyLocation, BodyScale3D);
 
 	// Direction
 	DirectionTowardEarth = BodyToEarth;
 }
 
-FString FSunInfo::ToString() const
+FString FStellarBody::ToString() const
 {
+	FString BaseString = Super::ToString();
+
 	return FString::Printf(TEXT(
-			"RightAscensionHours: %s\n"
-			"Declination: %s\n"
-			),
-			*UCelestialMaths::Conv_RightAscensionToString(RA),
-			*UCelestialMaths::Conv_DeclinationToString(DEC)
-		);
+			"%s"
+			"-----------------\n"
+			"ColorIndex: %.2f\n"
+			"HipparcosID %d\n"
+			"HenryDraperID %d\n"
+			"YaleBrightStarID %d\n"),
+		*BaseString,
+		ColorIndex,
+		HipparcosID,
+		HenryDraperID,
+		YaleBrightStarID
+	);
 }
 
 
-FString FStarInfo::ToString() const
+FString FPlanetaryBody::ToString() const
 {
+	FString BaseString = Super::ToString();
+
+	FString OrbitInfoString = "";
+	if (OrbitType == EOrbitType::VSOP87)
+	{
+		FString VSOP87BodyTypeName = FString("Unknown");
+		if (UEnum* EnumClass = StaticEnum<EVSOP87BodyType>())
+		{
+			VSOP87BodyTypeName = EnumClass->GetNameStringByValue(static_cast<int64>(VSOP87BodyType));
+		}
+
+		OrbitInfoString = FString::Printf(TEXT("VSOP87BodyType: %s"), *VSOP87BodyTypeName);
+	}
+	else if (OrbitType == EOrbitType::Elliptic)
+	{
+		// TODO
+		OrbitInfoString = "Elliptic Orbit";
+	}
+
+
 	return FString::Printf(TEXT(
-			"Name: %s\n"
-			"RightAscensionHours: %s\n"
-			"Declination: %s\n"
-			"Distance: %.0f PC\n"
-			"Magnitude: %.2f\n"
-			),
-			*Name,
-			*UCelestialMaths::Conv_RightAscensionToString(RA),
-			*UCelestialMaths::Conv_DeclinationToString(DEC),
-			DistanceInPC,
-			Magnitude
-		);
-	
+			"%s"
+			"-----------------\n"
+			"%s\n"
+			"ApparentDiameterDegrees: %s / %.8f°\n"
+			"ScaledApparentDiameterDegrees: %.8f°\n"
+			"Age: %.2f\n"
+			"IlluminationPercentage: %.2f\n"),
+		*BaseString,
+		*OrbitInfoString,
+		*UCelestialMaths::Conv_DeclinationToString(ApparentDiameterDegrees), ApparentDiameterDegrees,
+		ScaledApparentDiameterDegrees,
+		Age,
+		IlluminationPercentage
+	);
+}
+
+void FPlanetaryBody::ComputeUELocalTransform(double ExtraScale, double UEDistanceOverride)
+{
+	FCelestialBody::ComputeUELocalTransform(ExtraScale, UEDistanceOverride);
+
+	// We need to scale the Apparent diameter too here. 
+	ScaledApparentDiameterDegrees = ApparentDiameterDegrees * ExtraScale;
 }
diff --git a/Engine/Plugins/Experimental/CelestialVault/Source/CelestialVault/Private/CelestialMaths.cpp b/Engine/Plugins/Experimental/CelestialVault/Source/CelestialVault/Private/CelestialMaths.cpp
--- a/Engine/Plugins/Experimental/CelestialVault/Source/CelestialVault/Private/CelestialMaths.cpp
+++ b/Engine/Plugins/Experimental/CelestialVault/Source/CelestialVault/Private/CelestialMaths.cpp
@@ -4,8 +4,7 @@
 #include "CelestialMaths.h"
 #include "CelestialVault.h"
 #include "MathUtil.h"
-#include "vsop87a_milli.h"
-#include "vsop87a_full_EarthMoon.h"
+#include "VSOP87.h"
 
 #pragma region Static Members
 
@@ -25,18 +24,17 @@
 //   Y     =  -0.000000479966  +0.917482137087  -0.397776982902   Y
 //   Z FK5     0.000000000000  +0.397776982902  +0.917482137087   Z VSOP87A
 // 
-FMatrix UCelestialMaths::VSOPToJ2000 = FMatrix( 
+FMatrix const UCelestialMaths::VSOPEclipticToFK5Equatorial = FMatrix( 
 		FVector(+1.000000000000, +0.000000440360, -0.000000190919),	
 		FVector(-0.000000479966, +0.917482137087, -0.397776982902),	
 		FVector(+0.000000000000, +0.397776982902, +0.917482137087),		 
-		FVector(+0.000000000000, 0.000000000000, +0.000000000000)).GetTransposed();	// No Origin offset
+		FVector(+0.000000000000, 0.000000000000, +0.000000000000)).GetTransposed(); 	// No Origin offset
 
-double UCelestialMaths::SpeedOfLightMetersPerSeconds = 299792458.0;
-
-double UCelestialMaths::AstronomicalUnitsMeters = 149597870700.0;
-
-double UCelestialMaths::NewMoonReferenceJulianDate = 2460705.025196759; // Known new Moon was January 29th, 2025, at 12:36:17 UTC
-double UCelestialMaths::SynodicMonthAverage = 29.530588853;
+double const UCelestialMaths::SpeedOfLightMetersPerSeconds = 299792458.0;
+double const UCelestialMaths::AstronomicalUnitsMeters = 149597870700.0;
+double const UCelestialMaths::ParsecAstronomicalUnits = 206264.80624709636;
+double const UCelestialMaths::NewMoonReferenceJulianDate = 2460705.025196759; // Known new Moon was January 29th, 2025, at 12:36:17 UTC
+double const UCelestialMaths::SynodicMonthAverage = 29.530588853;
 
 #pragma endregion
 
@@ -58,138 +56,171 @@ FLinearColor UCelestialMaths::BVtoLinearColor(float BV)
 	return FLinearColor::MakeFromColorTemperature( Temperature);;
 }
 
+double UCelestialMaths::GetIlluminationPercentage(double NormalizedAge)
+{
+	return 0.5 * (1.0-FMath::Cos(2.0*UE_PI * NormalizedAge));
+}
+
 #pragma endregion
 
-#pragma region Planetary Bodies
+#pragma region VSOP87
 
-FVector UCelestialMaths::GetBodyLocation_FK5J2000_AU(FPlanetaryBodyInputData PlanetaryBody, double JulianDate)
+void UCelestialMaths::GetSolarSystemBodyLocationVelocity_VSOP87_AU(EVSOP87BodyType VSOP87BodyType, double JulianDate, FVector& LocationAU, FVector& Velocity)
 {
-	// From VSOP87.doc
-	//   Being given a Julian date JD expressed in dynamical time (TAI+32.184s) and a body (planets, Earth-Moon Barycenter, or Sun) associated to a version of the theory VSOP87 :
-	//     1/ select the file corresponding to the body and the version,
-	//     2/ read sequentially the terms of the series in the records of the file,
-	//     3/ apply for each term the formulae (1) or (2) with T=(JD-2451545)/365250,
-	//     4/ add up the terms so computed for every one coordinate.
+	// Warning - We need to use VSOP 87 Time for these functions! 
+	double VSOPTime = JulianDateToVSOP87Time(JulianDate);
 
-	// Convert time, because, the VSOP87 coordinates expect the time in TerrestrialTime
-	double TAI = JulianDateToInternationalAtomicTime(JulianDate); 
-	double TT = InternationalAtomicTimeToTerrestrialTime(TAI); // adds the 32.184s
-	double Time = JulianDateToJulianCenturies(SecondsToDay(TT)) / 10.0; // We divide by 10 because VSOP expects T : time expressed in Thousands of Julian Years (tjy) elapsed from J2000 (JD2451545.0).
-
-	
 	// Get the body coordinates in the inertial frame defined by the dynamical equinox and ecliptic J2000 (JD2451545.0).
-	double BodyXYZ[3] = { 0.0, 0.0, 0.0 }; 
-	switch (PlanetaryBody.OrbitType)
+	Velocity = FVector::ZeroVector;
+	LocationAU = FVector::ZeroVector;
+	
+	switch (VSOP87BodyType)
 	{
-	case EOrbitType::Mercury:
-		vsop87a_milli::getMercury(Time, BodyXYZ);
+	case EVSOP87BodyType::Mercury:
+		LocationAU = UVSOP87::GetMercuryLocation(VSOPTime, EVSOP87Accuracy::Milli);
 		break;
-	case EOrbitType::Venus:
-		vsop87a_milli::getVenus(Time, BodyXYZ);
+	case EVSOP87BodyType::Venus:
+		LocationAU = UVSOP87::GetVenusLocation(VSOPTime, EVSOP87Accuracy::Milli);
 		break;
-	case EOrbitType::Earth:
-		vsop87a_full::getEarth(Time, BodyXYZ);
+	case EVSOP87BodyType::Earth:
+		LocationAU = UVSOP87::GetEarthLocation(VSOPTime, EVSOP87Accuracy::Full);
+		Velocity = UVSOP87::GetEarthVelocity(VSOPTime, EVSOP87Accuracy::Full);
 		break;
-	case EOrbitType::Mars:
-		vsop87a_milli::getMars(Time, BodyXYZ);
+	case EVSOP87BodyType::Mars:
+		LocationAU = UVSOP87::GetMarsLocation(VSOPTime, EVSOP87Accuracy::Milli);
 		break;
-	case EOrbitType::Jupiter:
-		vsop87a_milli::getJupiter(Time, BodyXYZ);
+	case EVSOP87BodyType::Jupiter:
+		LocationAU = UVSOP87::GetJupiterLocation(VSOPTime, EVSOP87Accuracy::Milli);
 		break;
-	case EOrbitType::Saturn:
-		vsop87a_milli::getSaturn(Time, BodyXYZ);
+	case EVSOP87BodyType::Saturn:
+		LocationAU = UVSOP87::GetSaturnLocation(VSOPTime, EVSOP87Accuracy::Milli);
 		break;
-	case EOrbitType::Uranus:
-		vsop87a_milli::getUranus(Time, BodyXYZ);
+	case EVSOP87BodyType::Uranus:
+		LocationAU = UVSOP87::GetUranusLocation(VSOPTime, EVSOP87Accuracy::Milli);
 		break;
-	case EOrbitType::Neptune:
-		vsop87a_milli::getNeptune(Time, BodyXYZ);
+	case EVSOP87BodyType::Neptune:
+		LocationAU = UVSOP87::GetNeptuneLocation(VSOPTime, EVSOP87Accuracy::Milli);
 		break;
-	case EOrbitType::Moon:
-		{
-			// Special case for the Moon. VSOP Works by getting the Earth/Moon barycenter, and combine it with the earth XYZ
-			double EarthXYZ[3] = { 0.0, 0.0, 0.0 }; 
-			vsop87a_full::getEarth(Time, EarthXYZ);
-			double EMBXYZ[3] = { 0.0, 0.0, 0.0 };
-			vsop87a_full::getEmb(Time, EMBXYZ); 
-			vsop87a_full::getMoon(EarthXYZ, EMBXYZ, BodyXYZ);	
-		}
-	case EOrbitType::Elliptic:
-		{
-			// TODO - Use ellipsoid equation to compute locations
-		}
+	case EVSOP87BodyType::Moon:
+		LocationAU = UVSOP87::GetMoonLocation(VSOPTime, EVSOP87Accuracy::Full);
+		break;
+	case EVSOP87BodyType::Sun:
+		// Sun is at Origin, so returning 0 is what we want! 
 		break;
 	}
-
-	// Transform these VSOP coordinates into Equatorial Rectangular Coordinates (X, Y, Z)
-	// The FK5 is an equatorial coordinate system (coordinate system linked to the Earth) based on its J2000 position.
-	return VSOPToJ2000.TransformVector(FVector(BodyXYZ[0], BodyXYZ[1], BodyXYZ[2]));
 }
 
-FVector UCelestialMaths::GetBodyLocation_FK5J2000_AU_Relativistic(FVector ObserverBodyFK5J2000LocationAU, FPlanetaryBodyInputData PlanetaryBody, double JulianDate)
+FVector UCelestialMaths::GetSolarSystemBodyLocation_VSOP87_Relativistic(FVector ReferenceBodyLocation_Heliocentric_AU, EVSOP87BodyType VSOP87BodyType, double JulianDate)
 {
 	double JulianDateLightAdjusted = JulianDate;
 	FVector BodyPositionAU = FVector::ZeroVector;
 
 	for (int32 i = 0; i < 3; i++) // 3 iterations are good enough to converge
 	{
-		BodyPositionAU = GetBodyLocation_FK5J2000_AU(PlanetaryBody, JulianDateLightAdjusted);
-		double PlanetaryBodyDistanceAU = FVector::Distance(ObserverBodyFK5J2000LocationAU, BodyPositionAU);
+		FVector Velocities;
+		GetSolarSystemBodyLocationVelocity_VSOP87_AU(VSOP87BodyType, JulianDateLightAdjusted, BodyPositionAU, Velocities);
+		double PlanetaryBodyDistanceAU = FVector::Distance(ReferenceBodyLocation_Heliocentric_AU, BodyPositionAU);
 		double lightPropagationTimeInDays = SecondsToDay( AstronomicalUnitsToMeters(PlanetaryBodyDistanceAU) / SpeedOfLightMetersPerSeconds);
 		JulianDateLightAdjusted = JulianDate - lightPropagationTimeInDays;
 	}
 	return BodyPositionAU;
 }
 
-void UCelestialMaths::GetBodyCelestialCoordinatesAU(double JulianDate, FPlanetaryBodyInputData PlanetaryBody, double ObserverLatitude, double ObserverLongitude, double& RAHours, double& DECDegrees,
-	double& DistanceBodyToEarthAU, double& DistanceBodyToSunAU, double& DistanceEarthToSunAU)
+#pragma endregion
+
+
+#pragma region Celestial Bodies
+
+void UCelestialMaths::GetBodyCelestialCoordinatesAU(double JulianDate, EVSOP87BodyType VSOP87BodyType, double ObserverLatitude, double ObserverLongitude,
+	double& RAJ2000Hours, double& DECJ2000Degrees, double& RAHours, double& DECDegrees, double& DistanceBodyToEarthAU, double& DistanceBodyToSunAU, double& DistanceEarthToSunAU) 
 {
 	TRACE_CPUPROFILER_EVENT_SCOPE(UCelestialMaths::GetBodyCelestialCoordinatesAU);
 
+	// Compute the Precession and Nutation Matrices - The formulas expect a Date in Terrestrial Time
+	double TAISeconds = JulianDateToInternationalAtomicTime(JulianDate);
+	double TTSeconds = InternationalAtomicTimeToTerrestrialTime(TAISeconds);
+	double JulianDateTT = SecondsToDay(TTSeconds);
+	FMatrix PrecessionMatrix = GetPrecessionMatrix(JulianDateTT);
+	FMatrix NutationMatrix = GetNutationMatrix(JulianDateTT);
+
+	// Note about Coordinates systems
+
+	// The VSOP frame is a "heliocentric ecliptic frame of the J2000 equinox"
+	//    Origin = Sun (Heliocentric)
+	//    Coordinates Axes = Mean ecliptic & equinox of J2000.0
+	//       XY-plane: the mean ecliptic plane at epoch J2000.0
+	//       X-axis: pointing toward mean equinox at J2000.0
+	//       Z-axis: perpendicular to the ecliptic
+
+	// The FK5 system is a geocentric celestial reference frame:
+	//    Origin = Earth (GeoCentric)
+	//    Coordinates Axes = Mean equator and mean equinox of J2000.0.
+	//       X-axis Points toward the mean equinox of J2000.0 (intersection of Earth’s equator with the ecliptic at J2000).
+	//       Z-axis Points toward the mean north celestial pole at J2000.0 (Earth’s spin axis direction).
+	//       Y-axis Completes the right-handed coordinate system (90° east of X along the equator).
+	//    FK5 does not rotate with Earth’s daily rotation → the frame is fixed in space
+
+	// When getting the coordinates from any VSOP87 Call, it's important to change the axis frame to FK5J2000 using the VSOPEclipticToFK5Equatorial Matrix
+	// All the locations hereby defined will be using the FK5J2000 Frame, with HelioCentric or GeoCentric to precise the offset due to the Earth location
 	
-	RAHours = 0.0;
-	DECDegrees = 0.0;
-	DistanceBodyToEarthAU = 0.0; 
+	// Get Earth Location in VSOP Rectangular Coordinates (X, Y, Z)
+	FVector EarthVelocity_VSOP87 = FVector::ZeroVector;
+	FVector EarthLocation_VSOP87 = FVector::ZeroVector;
+	GetSolarSystemBodyLocationVelocity_VSOP87_AU(EVSOP87BodyType::Earth, JulianDate, EarthLocation_VSOP87, EarthVelocity_VSOP87);
 	
-	if (PlanetaryBody.OrbitType == EOrbitType::Earth)
-	{
-		// It makes no sense to ask for Earth location relative to the Earth, but handle the case anyway...
-		DistanceBodyToSunAU = AstronomicalUnitsMeters; 
-		return;
-	}
-
-	// TODO-Beta
-	// Take into account the Nutation and Precession effect, for better precisions, and also to have the right Heliocentric coordinates, otherwise the Phase Angle computations are wrong.
-	// (Not too important for Moon Phase yet)
+	// Transform to the FK5 J2000 Frame 
+	FVector EarthLocation_Heliocentric_FK5J2000 = VSOPEclipticToFK5Equatorial.TransformVector(EarthLocation_VSOP87); 
+	FVector EarthVelocity_FK5J2000 = VSOPEclipticToFK5Equatorial.TransformVector(EarthVelocity_VSOP87);
 	
-	// Get Earth Location in FK5J2000 Rectangular Coordinates (X, Y, Z)
-	FVector EarthLocationFK5J2000AU = GetBodyLocation_FK5J2000_AU(FPlanetaryBodyInputData::Earth, JulianDate);
-	UE_LOG(LogCelestialVault, Verbose, TEXT("Earth Location : FK5J2000 = %s"), *GetPreciseVectorString(EarthLocationFK5J2000AU));
-
 	// Get Body Location in FK5J2000 Rectangular Coordinates (X, Y, Z)
-	FVector BodyLocationFK5J2000AU = GetBodyLocation_FK5J2000_AU_Relativistic(EarthLocationFK5J2000AU, PlanetaryBody, JulianDate);
-	DistanceBodyToSunAU = BodyLocationFK5J2000AU.Length();
-	UE_LOG(LogCelestialVault, Verbose, TEXT("%s Location : FK5J2000 = %s"), *PlanetaryBody.Name, *GetPreciseVectorString(BodyLocationFK5J2000AU));
-
-	FVector BodyLocationGeoCentricAU = BodyLocationFK5J2000AU - EarthLocationFK5J2000AU;
-	FVector ObserverLocationGeocentricAU = GetObserverGeocentricLocationAU(ObserverLatitude, ObserverLongitude, 0, JulianDate );
-	FVector TopoCentricTargetAU = BodyLocationGeoCentricAU - ObserverLocationGeocentricAU;
-	DistanceEarthToSunAU = (ObserverLocationGeocentricAU-EarthLocationFK5J2000AU).Length();
+	FVector BodyLocation_VSOP87 = GetSolarSystemBodyLocation_VSOP87_Relativistic(EarthLocation_VSOP87, VSOP87BodyType, JulianDate);
+	FVector BodyLocation_Heliocentric_FK5J2000_AU = VSOPEclipticToFK5Equatorial.TransformVector(BodyLocation_VSOP87);
+
+	// Compute the Observer Location in the same, applying the effects of Precession matrix. 
+	FVector ObserverLocationGeoCentric_FK5J2000AU = GetObserverGeocentricLocationAU(ObserverLatitude, ObserverLongitude, 0, JulianDate );
+	ObserverLocationGeoCentric_FK5J2000AU = NutationMatrix.TransformPosition(ObserverLocationGeoCentric_FK5J2000AU);
+	ObserverLocationGeoCentric_FK5J2000AU = PrecessionMatrix.TransformPosition(ObserverLocationGeoCentric_FK5J2000AU); 
+
+	// In the FK5J2000AU Frame (S = Sun, B = Body, E = Earth, O = Observer
+	//   B
+	//
+	//
+	//            O
+	//  S        E
+	//
+	//  // Compute the OB vector
+	// OB = OE + EB
+	// OB = EB - EO
+	// OB = B-E - EO
+	// OB = BodyLocation_Heliocentric_FK5J2000_AU - EarthLocation_Heliocentric_FK5J2000 - ObserverLocationGeoCentric_FK5J2000AU
+	//
+	// From the above formula, it looks like some values are expressed in different frames (Heliocentric/GeoCentric), which still share the same axes. 
+	// It fine to do so, because we subtract the Body and Earth location, therefore the origin doesn't matter.
+	// The last argument, ObserverLocation GeoCentric is also expressed as a difference between the Observer location and the Earth center.    
+
+	// Convert the XYZ Coordinates to RA, DEC - This will be the J2000 ones!
+	// Topocentric means “from the point of view of a specific location on Earth.” TopoCentric is not Geocentric! 
+	FVector BodyLocation_TopoCentric_FK5J2000_AU = BodyLocation_Heliocentric_FK5J2000_AU - EarthLocation_Heliocentric_FK5J2000 - ObserverLocationGeoCentric_FK5J2000AU;
+	double RAJ2000Degrees; 
+	XYZToRADEC_RH(BodyLocation_TopoCentric_FK5J2000_AU, RAJ2000Degrees, DECJ2000Degrees, DistanceBodyToEarthAU);
+    RAJ2000Hours = RAJ2000Degrees / 15.0;
+
+	// Compute the RA and DEC of Date
+	// Transform the Body Topographic location with the Precession and Nutation Matrices at Date. 
+	FVector BodyLocation_TopoCentric_FK5J2000_Aberrated_AU = ComputeAberration(BodyLocation_TopoCentric_FK5J2000_AU, EarthVelocity_FK5J2000);
+	FVector BodyLocation_TopoCentric_FK5J2000_Precessed_AU = PrecessionMatrix.InverseTransformVector(BodyLocation_TopoCentric_FK5J2000_Aberrated_AU);
+	FVector BodyLocation_TopoCentric_FK5J2000_OfDate_AU = NutationMatrix.InverseTransformVector(BodyLocation_TopoCentric_FK5J2000_Precessed_AU);
 	
 	double RADegrees = 0.0;
-	XYZToRADEC_RH(TopoCentricTargetAU, RADegrees, DECDegrees, DistanceBodyToEarthAU);
+	XYZToRADEC_RH(BodyLocation_TopoCentric_FK5J2000_OfDate_AU, RADegrees, DECDegrees, DistanceBodyToEarthAU);
 	RAHours = RADegrees / 15.0;
 
-	// Debug Log...
-	UE_LOG(LogCelestialVault, Verbose,
-		TEXT("Planetary Body %s Location : RA = %s, DEC = %s, Radius = %f UA"),
-		*PlanetaryBody.Name,
-		*Conv_RightAscensionToString(RAHours),
-		*Conv_DeclinationToString(DECDegrees),
-		DistanceBodyToEarthAU );
+	DistanceBodyToEarthAU = (BodyLocation_Heliocentric_FK5J2000_AU - EarthLocation_Heliocentric_FK5J2000).Length();
+	DistanceBodyToSunAU = BodyLocation_Heliocentric_FK5J2000_AU.Length();
+	DistanceEarthToSunAU = EarthLocation_Heliocentric_FK5J2000.Length();
 }
 
-double UCelestialMaths::GetPlanetaryBodyMagnitude(FPlanetaryBodyInputData PlanetaryBody, double DistanceToSunAU, double DistanceToEarthAU, double DistanceEarthToSunAU, double& PhaseAngle)
+double UCelestialMaths::GetPlanetaryBodyMagnitude(EVSOP87BodyType VSOP87BodyType, double DistanceToSunAU, double DistanceToEarthAU, double DistanceEarthToSunAU, double& PhaseAngle)
 {
 	// From Computing Apparent Planetary Magnitudes for The Astronomical Almanac
 	// James L. Hilton US Naval Observatory
@@ -210,24 +241,21 @@ double UCelestialMaths::GetPlanetaryBodyMagnitude(FPlanetaryBodyInputData Planet
 	double PhaseAngle4 = PhaseAngle3 * PhaseAngle;
 	double PhaseAngle5 = PhaseAngle4 * PhaseAngle;
 	double PhaseAngle6 = PhaseAngle5 * PhaseAngle;
-	double PhaseFunction = 0.0; 
-	
-	switch (PlanetaryBody.OrbitType)
+	double PhaseFunction = 0.0;
+
+
+	switch (VSOP87BodyType)
 	{
-	case EOrbitType::Elliptic:
-		// Fictional Body. Return dummy magnitude
-		return 0.0;
-		break;
-	
-	case EOrbitType::Mercury:
+	case EVSOP87BodyType::Mercury:
 		ApparentMagnitude = -0.613;
-		PhaseFunction = 6.3280E-02 * PhaseAngle - 1.6336E-03 * PhaseAngle2 + 3.3644E-05 * PhaseAngle3 - 3.4265E-07 * PhaseAngle4 + 1.6893E-09 * PhaseAngle5 - 3.0334E-12 *PhaseAngle6;
+		PhaseFunction = 6.3280E-02 * PhaseAngle - 1.6336E-03 * PhaseAngle2 + 3.3644E-05 * PhaseAngle3 - 3.4265E-07 * PhaseAngle4 + 1.6893E-09 *
+			PhaseAngle5 - 3.0334E-12 * PhaseAngle6;
 		break;
-	case EOrbitType::Venus:
-		if ( PhaseAngle < 163.7)
+	case EVSOP87BodyType::Venus:
+		if (PhaseAngle < 163.7)
 		{
 			ApparentMagnitude = -4.384;
-			PhaseFunction = -1.044E-03 * PhaseAngle + 3.687E-04 * PhaseAngle2  - 2.814E-06 * PhaseAngle3 + 8.938E-09 * PhaseAngle4;
+			PhaseFunction = -1.044E-03 * PhaseAngle + 3.687E-04 * PhaseAngle2 - 2.814E-06 * PhaseAngle3 + 8.938E-09 * PhaseAngle4;
 		}
 		else // 163.7 <  α < 179 - let's go to 180...
 		{
@@ -235,85 +263,75 @@ double UCelestialMaths::GetPlanetaryBodyMagnitude(FPlanetaryBodyInputData Planet
 			PhaseFunction = -2.81914 * PhaseAngle + 8.39034E-03 * PhaseAngle2;
 		}
 		break;
-	case EOrbitType::Earth:
+	case EVSOP87BodyType::Earth:
 		ApparentMagnitude = -3.99;
 		PhaseFunction = -1.060E-3 * PhaseAngle + 2.054E-4 * PhaseAngle2;
 		break;
-	case EOrbitType::Mars:
-		if ( PhaseAngle < 50.0)
+	case EVSOP87BodyType::Mars:
+		if (PhaseAngle < 50.0)
 		{
 			ApparentMagnitude = -1.601;
-			PhaseFunction =  0.02267 * PhaseAngle - 0.0001302 * PhaseAngle2;
+			PhaseFunction = 0.02267 * PhaseAngle - 0.0001302 * PhaseAngle2;
 		}
-		else 
+		else
 		{
 			ApparentMagnitude = -0.367;
-			PhaseFunction = - 0.02573 * PhaseAngle + 0.0003445 * PhaseAngle2;
+			PhaseFunction = -0.02573 * PhaseAngle + 0.0003445 * PhaseAngle2;
 		}
 		break;
-	case EOrbitType::Jupiter:
-		if ( PhaseAngle < 12.0)
+	case EVSOP87BodyType::Jupiter:
+		if (PhaseAngle < 12.0)
 		{
-			ApparentMagnitude =  -9.395;
-			PhaseFunction =  -3.7E-04 * PhaseAngle - 6.16E-04 * PhaseAngle2;
+			ApparentMagnitude = -9.395;
+			PhaseFunction = -3.7E-04 * PhaseAngle - 6.16E-04 * PhaseAngle2;
 		}
-		else // 12 < α < 130 - The phase curve of Jupiter as seen from Earth cannot exceed α = 12 so we should be good. Add this one just in case...
+		else
+		// 12 < α < 130 - The phase curve of Jupiter as seen from Earth cannot exceed α = 12 so we should be good. Add this one just in case...
 		{
 			ApparentMagnitude = -9.428;
 			PhaseFunction = -2.5 * FMath::LogX(10.0,
 				1.0
-				-1.507 * (PhaseAngle/180.0)
-				-0.363 * FMath::Pow(PhaseAngle / 180.0, 2.0)
-				-0.062 * FMath::Pow(PhaseAngle / 180.0, 3.0)
-				+2.809 * FMath::Pow(PhaseAngle / 180.0, 4.0)
-				-1.876 * FMath::Pow(PhaseAngle / 180.0, 5.0)); 
+				- 1.507 * (PhaseAngle / 180.0)
+				- 0.363 * FMath::Pow(PhaseAngle / 180.0, 2.0)
+				- 0.062 * FMath::Pow(PhaseAngle / 180.0, 3.0)
+				+ 2.809 * FMath::Pow(PhaseAngle / 180.0, 4.0)
+				- 1.876 * FMath::Pow(PhaseAngle / 180.0, 5.0));
 		}
 		break;
-	case EOrbitType::Saturn:
+	case EVSOP87BodyType::Saturn:
 		// Keep it simple and ignore the ring effects
-		if ( PhaseAngle < 6.0)
+		if (PhaseAngle < 6.0)
 		{
-			ApparentMagnitude =  -8.95;
-			PhaseFunction =  - 3.7E-04 * PhaseAngle + 6.16E-04 * PhaseAngle2;
+			ApparentMagnitude = -8.95;
+			PhaseFunction = -3.7E-04 * PhaseAngle + 6.16E-04 * PhaseAngle2;
 		}
 		else // 6 < α < 150 -
 		{
-			ApparentMagnitude = - 8.94;
+			ApparentMagnitude = -8.94;
 			PhaseFunction = 2.446E-4 * PhaseAngle + 2.672E-4 * PhaseAngle2 - 1.505E-6 * PhaseAngle3 + 4.767E-9 * PhaseAngle4;
 		}
 		break;
-	case EOrbitType::Uranus:
+	case EVSOP87BodyType::Uranus:
 		ApparentMagnitude = -7.19;
-		PhaseFunction =  -8.4E-04 * 82.0;  // PhaseAngle doesn't have any impact, and 82 is the most important planetographic latitude
+		PhaseFunction = -8.4E-04 * 82.0; // PhaseAngle doesn't have any impact, and 82 is the most important planetographic latitude
 		break;
-	case EOrbitType::Neptune:
+	case EVSOP87BodyType::Neptune:
 		ApparentMagnitude = -7.00;
-		PhaseFunction =  7.944E-3 * PhaseAngle + 9.617E-5 * PhaseAngle2;
+		PhaseFunction = 7.944E-3 * PhaseAngle + 9.617E-5 * PhaseAngle2;
 		break;
-	case EOrbitType::Moon:
+	case EVSOP87BodyType::Moon:
 		ApparentMagnitude = -12.73;
 		PhaseFunction = 0.026 * PhaseAngle + 4E-9 * PhaseAngle4;
 		break;
+	case EVSOP87BodyType::Sun:
+		return -26.74;
+		break;
 	default: ;
 	}
 
 	return DistanceFactor + ApparentMagnitude + PhaseFunction;
 }
 
-double UCelestialMaths::GetMoonNormalizedAgeSimple(double JulianDate)
-{
-	double DeltaDays = JulianDate - NewMoonReferenceJulianDate;
-	double MoonAgeDays = ModPositive(DeltaDays, SynodicMonthAverage);
-
-	return MoonAgeDays / SynodicMonthAverage;
-}
-
-double UCelestialMaths::GetIlluminationPercentage(double NormalizedAge)
-{
-	return 0.5 * (1.0-FMath::Cos(2.0*UE_PI * NormalizedAge));
-}
-
-
 #pragma endregion
 
 #pragma region Time
@@ -325,7 +343,7 @@ FDateTime UCelestialMaths::LocalTimeToUTCTime(FDateTime LocalTime, double TimeZo
 		TimeZoneOffset += 1.0;
 	}
 
-	return LocalTime - FTimespan(TimeZoneOffset,0,0);
+	return LocalTime - FTimespan::FromHours(TimeZoneOffset);
 }
 
 FDateTime UCelestialMaths::UTCTimeToLocalTime(FDateTime UTCTime, double TimeZoneOffset, bool IsDst)
@@ -335,7 +353,7 @@ FDateTime UCelestialMaths::UTCTimeToLocalTime(FDateTime UTCTime, double TimeZone
 		TimeZoneOffset += 1.0;
 	}
 
-	return UTCTime + FTimespan(TimeZoneOffset,0,0);
+	return UTCTime + FTimespan::FromHours(TimeZoneOffset);
 }
 
 double UCelestialMaths::UTCDateTimeToJulianDate(FDateTime UTCDateTime)
@@ -403,7 +421,7 @@ FDateTime UCelestialMaths::JulianDateToUTCDateTime(double JulianDate)
 	int32 D = FloorForJulianDate(365.25 * C);
 	int32 E = FloorForJulianDate((B - D) / 30.6001);
 
-	int32 day = B - D - static_cast<int32>(30.6001 * E) + FractionalPart;
+	int32 day = B - D - static_cast<int32>(30.6001 * E);
 	int32 month = E - 1;
 	if (E > 13)
 	{
@@ -429,10 +447,10 @@ FDateTime UCelestialMaths::JulianDateToUTCDateTime(double JulianDate)
 
 double UCelestialMaths::DateTimeToGreenwichMeanSiderealTime(FDateTime UTCDateTime)
 {
-	double JulianDate = UTCDateTime.GetJulianDay();
+	double JulianDate = UTCDateTimeToJulianDate(UTCDateTime);
 	double JulianCentury = JulianDateToJulianCenturies(JulianDate);
-	double SideralTimeAt0h = 100.46061837 + 36000.770053608 * JulianCentury + 0.000387933 * JulianCentury * JulianCentury + 1.0/38710000.0 * JulianCentury * JulianCentury * JulianCentury;
-	return ModPositive(SideralTimeAt0h, 360.0);
+	double SideralTime = 280.46061837 + 360.98564736629*(JulianDate-2451545.0) + 0.000387933 * JulianCentury * JulianCentury - 1.0/38710000.0 * JulianCentury * JulianCentury * JulianCentury;
+	return ModPositive(SideralTime, 360.0);
 }
 
 double UCelestialMaths::JulianDateToGreenwichMeanSiderealTime(double JulianDate) {
@@ -457,7 +475,7 @@ double UCelestialMaths::JulianDateToGreenwichMeanSiderealTime(double JulianDate)
 
 double UCelestialMaths::LocalSideralTime(double LongitudeDegrees, double GreenwichMeanSideralTime)
 {
-	return GreenwichMeanSideralTime + LongitudeDegrees;
+	return ModPositive(GreenwichMeanSideralTime + LongitudeDegrees, 360.0);
 }
 
 double UCelestialMaths::JulianDateToGreenwichApparentSiderealTime(double JulianDate)
@@ -541,6 +559,18 @@ double UCelestialMaths::JulianDateToInternationalAtomicTime(double JulianDate)
 	return GetLeapSeconds(JulianDate) + DaysToSeconds(JulianDate);
 }
 
+double UCelestialMaths::JulianDateToVSOP87Time(double JulianDate)
+{
+	// Convert time, because, the VSOP87 coordinates expect the time in TerrestrialTime
+	double TAI = JulianDateToInternationalAtomicTime(JulianDate); 
+	double TT = InternationalAtomicTimeToTerrestrialTime(TAI); // adds the 32.184s
+
+	// Caution here. In the VSOP87 equations, Time is not Julian centuries! 
+	// The denominator is 365250, not 36525.
+	// That’s because VSOP87 expresses time in units of 10,000 Julian years (1 Julian millennium = 100 centuries = 36,525 × 100 days).
+	return (SecondsToDay(TT)-2451545.0) / 365250.0; // We divide by 10 because VSOP expects T : time expressed in Thousands of Julian Years (tjy) elapsed from J2000 (JD2451545.0).
+}
+
 #pragma endregion
 
 #pragma region Angles
@@ -632,19 +662,9 @@ FVector UCelestialMaths::GetObserverGeocentricLocationAU(double Latitude,double
 	const FVector ObserverECEF = GeodeticLatLonToECEFXYZAU(Latitude, Longitude, Altitude);
 	const double GreenwichApparentSiderealTime = JulianDateToGreenwichApparentSiderealTime(JulianDate);
 
-	// Compute cosine and sine of the angle
-	float CosTheta = FMath::Cos(-FMath::DegreesToRadians(GreenwichApparentSiderealTime));
-	float SinTheta = FMath::Sin(-FMath::DegreesToRadians(GreenwichApparentSiderealTime));
-	
-	// Construct the 3x3 rotation matrix to rotate the ECEF position around the Earth axis, depending on the GAST
-	FMatrix RotationMatrix = FMatrix(
-			FPlane(CosTheta, -SinTheta, 0.0f, 0.0f), // Row 1
-			FPlane(SinTheta,  CosTheta, 0.0f, 0.0f), // Row 2
-			FPlane(0.0f,      0.0f,     1.0f, 0.0f), // Row 3
-			FPlane(0.0f,      0.0f,     0.0f, 1.0f)  // Row 4 (required for FMatrix, but ignored here)
-		);
-	
-	return RotationMatrix.TransformVector(ObserverECEF);
+	// Construct the 4x4 rotation matrix to rotate the ECEF position around the Earth axis, depending on the GAST
+	FMatrix RotationMatrix = GetRotationZMatrix(-FMath::DegreesToRadians(GreenwichApparentSiderealTime));
+	return RotationMatrix.InverseTransformPosition(ObserverECEF);
 }
 
 double UCelestialMaths::EquationOfTheEquinoxes(double JulianDate) {
@@ -665,19 +685,19 @@ double UCelestialMaths::EquationOfTheEquinoxes(double JulianDate) {
 			-0.000000576 * JulianCenturies * JulianCenturies * JulianCenturies * JulianCenturies +
 			-0.0000000434 * JulianCenturies * JulianCenturies * JulianCenturies * JulianCenturies * JulianCenturies);
 
-	double DeltaPsi;
-	double DeltaEpsilon; 
-	Nutation2000BTruncated(JulianDate, DeltaPsi, DeltaEpsilon);
-	return DeltaPsi * FMath::Cos(FMath::DegreesToRadians(MeanObliquity + DeltaEpsilon));
+	double DeltaPsiDegrees;
+	double DeltaEpsilonDegrees; 
+	Nutation2000BTruncated(JulianCenturies, DeltaPsiDegrees, DeltaEpsilonDegrees);
+	return DeltaPsiDegrees * FMath::Cos(FMath::DegreesToRadians(MeanObliquity + DeltaEpsilonDegrees));
 }
 
-void UCelestialMaths::Nutation2000BTruncated(double JulianDate, double& DeltaPsi, double& DeltaEpsilon)
+void UCelestialMaths::Nutation2000BTruncated(double JulianCenturies, double& DeltaPsiDegrees, double& DeltaEpsilonDegrees)
 {
 	//The IAU Resolutions on Astronomical Reference Systems, Time Scales, and Earth Rotation Models Explanation and Implementation (George H. Kaplan)
 	//https://arxiv.org/pdf/astro-ph/0602086.pdf
 	//IAU 2000B Nutation truncated to 6 terms
 
-	const double JC = JulianDateToJulianCenturies(JulianDate);
+	const double JC = JulianCenturies;
 	const double JC2 = JC * JC;
 	const double JC3 = JC * JC2;
 	const double JC4 = JC * JC3;
@@ -728,11 +748,11 @@ void UCelestialMaths::Nutation2000BTruncated(double JulianDate, double& DeltaPsi
 	DeltaPsiArcSeconds += (-172064161.0 + -174666.0 * JC) * FMath::Sin(Omega) + 33386.0 * FMath::Cos(Omega);
 	DeltaEpsilonArcSeconds += (92052331.0 + 9086.0 * JC) * FMath::Cos(Omega) + 15377.0 * FMath::Sin(Omega);
 
-	DeltaPsi = ArcsecondsToDegrees(DeltaPsiArcSeconds / 10000000.0);
-	DeltaEpsilon = ArcsecondsToDegrees(DeltaEpsilonArcSeconds / 10000000.0);
+	DeltaPsiDegrees = ArcsecondsToDegrees(DeltaPsiArcSeconds / 10000000.0);
+	DeltaEpsilonDegrees = ArcsecondsToDegrees(DeltaEpsilonArcSeconds / 10000000.0);
 }
 
-FTransform UCelestialMaths::GetPlanetCenterTransform(double Latitude, double Longitude, double Altitude)
+FTransform UCelestialMaths::GetEarthCenterTransformECEF(double Latitude, double Longitude, double Altitude)
 {
 	// Compute the Location part 
 	FVector ECEFLocation = GeodeticLatLonToECEFXYZAU(Latitude, Longitude, Altitude) * AstronomicalUnitsMeters;
@@ -787,154 +807,169 @@ FTransform UCelestialMaths::GetPlanetCenterTransform(double Latitude, double Lon
 	return FTransform(AxisMatrix.Inverse());
 }
 
-#pragma endregion
+FTransform UCelestialMaths::GetEarthCenterTransformUEFrame(double Latitude, double Longitude, double Altitude)
+{
+	// Get the transformation in proper Celestial World units (meters, right handed)
+	FTransform ECEFFrameToWorldFrame = GetEarthCenterTransformECEF(Latitude, Longitude, 0);
+
+	// UE Frame are expressed in Left-handed coordinates, and units are in meters - Convert to UE Transform 
+	FMatrix WorldFrameToUEFrame = FMatrix( 
+		FVector(1.0, 0.0, 0.0),			// Easting (X) is UE World X
+		FVector(0.0, -1.0, 0.0),			// Northing (Y) is UE World -Y because of left-handed convention
+		FVector(0.0, 0.0, 1.0),			// Up (Z) is UE World Z 
+		FVector(0.0, 0.0, 0.0));	// No Origin offset
+	FMatrix UEFrameToWorldFrame = WorldFrameToUEFrame.Inverse();
+
+	// Update the rotation part
+	FMatrix TransformMatrix = UEFrameToWorldFrame * ECEFFrameToWorldFrame.ToMatrixNoScale() * WorldFrameToUEFrame;
+	// Get Origin, and convert UE units to meters
+	FVector UEOrigin = TransformMatrix.GetOrigin() * FVector(100.0, 100.0, 100.0);
+	TransformMatrix.SetOrigin(UEOrigin);
+
+	// Apply the transform
+	return FTransform(TransformMatrix);
+}
 
-#pragma region Sun
+FMatrix UCelestialMaths::GetPrecessionMatrix(double JulianDate)
+{
+	// Fukushima-Williams IAU 2006
+	double t =  (JulianDate - 2451545.5) / 36525.0; // We don't use the Julian Centuries because precession expects JD TT
+
+	double gamma = ArcsecondsToRadians(
+		-0.052928 +
+		10.556378 * t +
+		0.4932044 * t * t -
+		0.00031238 * t * t * t -
+		0.000002788 * t * t * t * t +
+		0.000000026 * t * t * t * t * t);
+	double phi = ArcsecondsToRadians(
+		84381.412819 -
+		46.811016 * t +
+		0.0511268 * t * t +
+		0.00053289 * t * t * t -
+		0.00000044 * t * t * t * t -
+		0.0000000176 * t * t * t * t * t);
+	double psi = ArcsecondsToRadians(
+		-0.041775 +
+		5038.481484 * t +
+		1.5584175 * t * t -
+		0.00018522 * t * t * t -
+		0.000026452 * t * t * t * t -
+		0.0000000148 * t * t * t * t * t);
+	double eps = ArcsecondsToRadians(
+		84381.406 -
+		46.836769 * t -
+		0.0001831 * t * t +
+		0.0020034 * t * t * t -
+		0.000000576 * t * t * t * t -
+		0.0000000434 * t * t * t * t * t);
+
+	FMatrix a = GetRotationZMatrix(gamma);
+	FMatrix a2 = FRotationMatrix::Make(FRotator(0, FMath::RadiansToDegrees(gamma), 0));
+	
+	FMatrix Mphi = GetRotationXMatrix(phi);
+	FMatrix Mphi2 = FRotationMatrix::Make(FRotator(FMath::RadiansToDegrees(phi), 0, 0));
+	
+	FMatrix b = Mphi * a;
+	FMatrix b2 = Mphi2 * a2;
+	FMatrix Mpsi = GetRotationZMatrix(-psi);
+	FMatrix c = Mpsi * b;
+	FMatrix Meps = GetRotationXMatrix(-eps);
+	FMatrix d = Meps * c;
+	
+	return d;
+}
 
-FSunInfo UCelestialMaths::GetSunInformation(double JulianDate, double ObserverLatitude, double ObserverLongitude)
+FMatrix UCelestialMaths::GetNutationMatrix(double JulianDate)
 {
-    	FSunInfo SunInfo;
+	double t =  (JulianDate - 2451545.5) / 36525.0; // We don't use the Julian Centuries because precession expects JD TT
 
-		// TODO_Beta
-		// This code comes from the former SunPosition Plugin. Add the missing properties but it will require Time Zone, DST, and such... 
+	double DeltaPsiDegrees = 0.0;
+	double DeltaEpsilonDegrees = 0.0;
+	Nutation2000BTruncated(t, DeltaPsiDegrees, DeltaEpsilonDegrees);
+	double DeltaPsiRadians = FMath::DegreesToRadians(DeltaPsiDegrees);
+	double DeltaEpsilonRadians = FMath::DegreesToRadians(DeltaEpsilonDegrees);
 	
+	double eps = ArcsecondsToRadians(
+		84381.406 -
+		46.836769 * t -
+		0.0001831 * t * t +
+		0.0020034 * t * t * t -
+		0.000000576 * t * t * t * t -
+		0.0000000434 * t * t * t * t * t);
+
+	FMatrix a = GetRotationXMatrix(eps);
+	FMatrix b = GetRotationZMatrix(-DeltaPsiRadians) * a;
+	FMatrix c = GetRotationXMatrix(-(eps + DeltaEpsilonRadians)) * b;
+
+	return c;
+}
 
-	
-		// if (!FDateTime::Validate(Year, Month, Day, Hours, Minutes, Seconds, Milliseconds))
-		// {
-		// 	return SunInfo;
-		// }
-		//
-		// FDateTime CalcTime(Year, Month, Day, Hours, Minutes, Seconds, Milliseconds);
-		//
-		// double TimeOffset = TimeZone;
-		// if (bIsDaylightSavingTime)
-		// {
-		// 	TimeOffset += 1.0;
-		// }
-
-	    double LatitudeRad = FMath::DegreesToRadians(ObserverLatitude);
-
-	    // Get the julian day (number of days since Jan 1st of the year 4713 BC)
-	    double JulianDay = JulianDate;
-	    double JulianCentury = JulianDateToJulianCenturies(JulianDate);
-
-	    // Get the sun's mean longitude , referred to the mean equinox of julian date
-	    double GeomMeanLongSunDeg = FMath::Fmod(280.46646 + JulianCentury * (36000.76983 + JulianCentury * 0.0003032), 360.0);
-	    double GeomMeanLongSunRad = FMath::DegreesToRadians(GeomMeanLongSunDeg);
-
-	    // Get the sun's mean anomaly
-	    double GeomMeanAnomSunDeg = 357.52911 + JulianCentury * (35999.05029 - 0.0001537 * JulianCentury);
-	    double GeomMeanAnomSunRad = FMath::DegreesToRadians(GeomMeanAnomSunDeg);
-
-	    // Get the earth's orbit eccentricity
-	    double EccentEarthOrbit = 0.016708634 - JulianCentury * (0.000042037 + 0.0000001267 * JulianCentury);
-
-	    // Get the sun's equation of the center
-	    double SunEqOfCtr = FMath::Sin(GeomMeanAnomSunRad) * (1.914602 - JulianCentury * (0.004817 + 0.000014 * JulianCentury))
-		    + FMath::Sin(2.0 * GeomMeanAnomSunRad) * (0.019993 - 0.000101 * JulianCentury)
-		    + FMath::Sin(3.0 * GeomMeanAnomSunRad) * 0.000289;
-		
-		// Get the sun's true longitude
-		double SunTrueLongDeg = GeomMeanLongSunDeg + SunEqOfCtr;
+FVector UCelestialMaths::ComputeAberration(const FVector& TopocentricTargetLocationAU, const FVector& EarthVelocity)
+{
+	//"MEAN AND APPARENT PLACE COMPUTATIONS IN THE NEW IAU SYSTEM. III. APPARENT, TOPOCENTRIC, AND ASTROMETRIC PLACES OF PLANETS AND STARS"
+	//G. H. Kaplan, J. A. Hughes, P. K. Seidelmann, and C. A. Smith
+	//U. S. Naval Observatory, Washington, DC 20392
 
-		// Get the sun's true anomaly
-		//	double SunTrueAnomDeg = GeomMeanAnomSunDeg + SunEqOfCtr;
-		//	double SunTrueAnomRad = FMath::DegreesToRadians(SunTrueAnomDeg);
+	//http://articles.adsabs.harvard.edu/pdf/1989AJ.....97.1197K
 
-		// Get the earth's distance from the sun
-		//	double SunRadVectorAUs = (1.000001018*(1.0 - EccentEarthOrbit*EccentEarthOrbit)) / (1.0 + EccentEarthOrbit*FMath::Cos(SunTrueAnomRad));
-		
-		// Get the sun's apparent longitude
-		double SunAppLongDeg = SunTrueLongDeg - 0.00569 - 0.00478*FMath::Sin(FMath::DegreesToRadians(125.04 - 1934.136*JulianCentury));
-		double SunAppLongRad = FMath::DegreesToRadians(SunAppLongDeg);
+	double C = 173.1446326846693; //Speed of light in AU per Day
 
-		// Get the earth's mean obliquity of the ecliptic
-		double MeanObliqEclipticDeg = 23.0 + (26.0 + ((21.448 - JulianCentury*(46.815 + JulianCentury*(0.00059 - JulianCentury*0.001813)))) / 60.0) / 60.0;
-		
-		// Get the oblique correction
-		double ObliqCorrDeg = MeanObliqEclipticDeg + 0.00256*FMath::Cos(FMath::DegreesToRadians(125.04 - 1934.136*JulianCentury));
-		double ObliqCorrRad = FMath::DegreesToRadians(ObliqCorrDeg);
+	//Eq 16
+	double t = TopocentricTargetLocationAU.Length() / C;
+	double B = EarthVelocity.Length() / C;
+	double cosD = TopocentricTargetLocationAU.Dot(EarthVelocity) / (TopocentricTargetLocationAU.Length() * EarthVelocity.Length());
+	double y = FMath::Sqrt(1 - B * B);
+	double f1 = B * cosD;
+	double f2 = (1 + f1 / (1 + y)) * t;
 
-		// Get the sun's right ascension
-		double SunRtAscenRad = FMath::Atan2(FMath::Cos(ObliqCorrRad)*FMath::Sin(SunAppLongRad), FMath::Cos(SunAppLongRad));
-		double SunRtAscenDeg = FMath::RadiansToDegrees(SunRtAscenRad);
-		
-		// Get the sun's declination
-		double SunDeclinRad = FMath::Asin(FMath::Sin(ObliqCorrRad)*FMath::Sin(SunAppLongRad));
-		double SunDeclinDeg = FMath::RadiansToDegrees(SunDeclinRad);
+	//Eq 17
+	FVector AberrationLocation = ( y * TopocentricTargetLocationAU + f2 * EarthVelocity) / (1+f1);
+	return AberrationLocation;
+}
 
-		double VarY = FMath::Pow(FMath::Tan(ObliqCorrRad / 2.0), 2.0);
-		
-		// Get the equation of time
-		double EqOfTimeMinutes = 4.0 * FMath::RadiansToDegrees(VarY*FMath::Sin(2.0 * GeomMeanLongSunRad) - 2.0 * EccentEarthOrbit*FMath::Sin(GeomMeanAnomSunRad) + 4.0 * EccentEarthOrbit*VarY*FMath::Sin(GeomMeanAnomSunRad)*FMath::Cos(2.0 * GeomMeanLongSunRad) - 0.5*VarY*VarY*FMath::Sin(4.0 * GeomMeanLongSunRad) - 1.25*EccentEarthOrbit*EccentEarthOrbit*FMath::Sin(2.0 * GeomMeanAnomSunRad));
-
-		// Get the hour angle of the sunrise
-		double HASunriseDeg = FMath::RadiansToDegrees(FMath::Acos(FMath::Cos(FMath::DegreesToRadians(90.833)) / (FMath::Cos(LatitudeRad)*FMath::Cos(SunDeclinRad)) - FMath::Tan(LatitudeRad)*FMath::Tan(SunDeclinRad)));
-		//	double SunlightDurationMinutes = 8.0 * HASunriseDeg;
-
-		// // Get the local time of the sun's rise and set
-		// double SolarNoonLST = (720.0 - 4.0 * ObserverLongitude - EqOfTimeMinutes + TimeOffset * 60.0) / 1440.0;
-		// double SunriseTimeLST = SolarNoonLST - HASunriseDeg * 4.0 / 1440.0;
-		// double SunsetTimeLST = SolarNoonLST + HASunriseDeg * 4.0 / 1440.0;
-		//
-		// // Get the true solar time
-		// double TrueSolarTimeMinutes = FMath::Fmod(CalcTime.GetTimeOfDay().GetTotalMinutes() + EqOfTimeMinutes + 4.0 * ObserverLongitude - 60.0 * TimeOffset, 1440.0);
-		//
-		// // Get the hour angle of current time
-		// double HourAngleDeg = TrueSolarTimeMinutes < 0 ? TrueSolarTimeMinutes / 4.0 + 180 : TrueSolarTimeMinutes / 4.0 - 180.0;
-		// double HourAngleRad = FMath::DegreesToRadians(HourAngleDeg);
-		//
-		// // Get the solar zenith angle
-		// double SolarZenithAngleRad = FMath::Acos(FMath::Sin(LatitudeRad)*FMath::Sin(SunDeclinRad) + FMath::Cos(LatitudeRad)*FMath::Cos(SunDeclinRad)*FMath::Cos(HourAngleRad));
-		// double SolarZenithAngleDeg = FMath::RadiansToDegrees(SolarZenithAngleRad);
-		//
-		// // Get the sun elevation
-		// double SolarElevationAngleDeg = 90.0 - SolarZenithAngleDeg;
-		// double SolarElevationAngleRad = FMath::DegreesToRadians(SolarElevationAngleDeg);
-		// double TanOfSolarElevationAngle = FMath::Tan(SolarElevationAngleRad);
-		//
-		// // Get the approximated atmospheric refraction
-		// double ApproxAtmosphericRefractionDeg = 0.0;
-		// if (SolarElevationAngleDeg <= 85.0)
-		// {
-		// 	if (SolarElevationAngleDeg > 5.0)
-		// 	{
-		// 		ApproxAtmosphericRefractionDeg = 58.1 / TanOfSolarElevationAngle - 0.07 / FMath::Pow(TanOfSolarElevationAngle, 3) + 0.000086 / FMath::Pow(TanOfSolarElevationAngle, 5) / 3600.0;
-		// 	}
-		// 	else
-		// 	{
-		// 		if (SolarElevationAngleDeg > -0.575)
-		// 		{
-		// 			ApproxAtmosphericRefractionDeg = 1735.0 + SolarElevationAngleDeg * (-518.2 + SolarElevationAngleDeg * (103.4 + SolarElevationAngleDeg * (-12.79 + SolarElevationAngleDeg * 0.711)));
-		// 		}
-		// 		else
-		// 		{
-		// 			ApproxAtmosphericRefractionDeg = -20.772 / TanOfSolarElevationAngle;
-		// 		}
-		// 	}
-		// 	ApproxAtmosphericRefractionDeg /= 3600.0;
-		// }
-		//◘
-		// // Get the corrected solar elevation
-		// double SolarElevationcorrectedforatmrefractionDeg = SolarElevationAngleDeg + ApproxAtmosphericRefractionDeg;
-		//
-		// // Get the solar azimuth 
-		// double tmp = FMath::RadiansToDegrees(FMath::Acos(((FMath::Sin(LatitudeRad)*FMath::Cos(SolarZenithAngleRad)) - FMath::Sin(SunDeclinRad)) / (FMath::Cos(LatitudeRad)*FMath::Sin(SolarZenithAngleRad))));
-		// double SolarAzimuthAngleDegcwfromN = HourAngleDeg > 0.0 ? FMath::Fmod(tmp + 180.0, 360.0) : FMath::Fmod(540.0 - tmp, 360.0);
-
-
-		// offset elevation angle to fit with UE coords system
-    	SunInfo.RA = SunRtAscenDeg / 15.0;
-		SunInfo.DEC = SunDeclinDeg; 
-		// SunInfo.Elevation = 180.0 + SolarElevationAngleDeg;
-		// SunInfo.CorrectedElevation = 180.0 + SolarElevationcorrectedforatmrefractionDeg;
-		// SunInfo.Azimuth = SolarAzimuthAngleDegcwfromN;
-		// SunInfo.SolarNoon = FTimespan::FromDays(SolarNoonLST);
-		// SunInfo.SunriseTime = FTimespan::FromDays(SunriseTimeLST);
-		// SunInfo.SunsetTime = FTimespan::FromDays(SunsetTimeLST);
-    	return SunInfo;
-	}
+#pragma endregion
+
+#pragma region Sun & Moon
+
+FStellarBody UCelestialMaths::GetSunInformation(double JulianDate, double ObserverLatitude, double ObserverLongitude)
+{
+	// Get the Celestial Body Coordinates
+	double RAJ2000Hours;
+	double DECJ2000Degrees;
+	double RAHours;
+	double DECDegrees;
+	double DistanceBodyToEarthAU;
+	double DistanceBodyToSunAU;
+	double DistanceEarthToSunAU;
+	GetBodyCelestialCoordinatesAU(JulianDate, EVSOP87BodyType::Sun, ObserverLatitude, ObserverLongitude, RAJ2000Hours, DECJ2000Degrees, RAHours, DECDegrees, DistanceBodyToEarthAU, DistanceBodyToSunAU, DistanceEarthToSunAU);  
+
+	// Set the Sun properties. Some properties can be hardcoded for the Sun (CI, Magnitude, Radius) 
+	FStellarBody SunInfo;
+	SunInfo.ColorIndex = 0.656;
+	SunInfo.Color = UCelestialMaths::BVtoLinearColor(SunInfo.ColorIndex);
+	SunInfo.OrbitType = EOrbitType::VSOP87;
+	SunInfo.RAJ2000 = RAJ2000Hours;
+	SunInfo.DECJ2000 = DECJ2000Degrees;
+	SunInfo.RA = RAHours;
+	SunInfo.DEC = DECDegrees;
+	SunInfo.DistanceInAU = DistanceEarthToSunAU;
+	SunInfo.Magnitude = -26.74;
+	SunInfo.Radius = 695700;
+
+	// Make sure we build the UE Transformation Cache
+	SunInfo.ComputeUELocalTransform(); // No change in scale or distance for the Sun - Default arguments
+	
+	return SunInfo;
+}
+
+double UCelestialMaths::GetMoonNormalizedAgeSimple(double JulianDate)
+{
+	double DeltaDays = JulianDate - NewMoonReferenceJulianDate;
+	double MoonAgeDays = ModPositive(DeltaDays, SynodicMonthAverage);
+
+	return MoonAgeDays / SynodicMonthAverage;
+}
 
 #pragma endregion
 
@@ -942,6 +977,15 @@ FSunInfo UCelestialMaths::GetSunInformation(double JulianDate, double ObserverLa
 
 void UCelestialMaths::XYZToRADEC_RH(FVector XYZ, double& RADegrees, double& DECDegrees, double& Radius)
 {
+	if (XYZ.IsNearlyZero())
+	{
+		UE_LOG(LogCelestialVault, Error, TEXT("Zero-length XYZ vector in XYZToRADEC_RH - Check your inputs"));
+		RADegrees = 0.0;
+		DECDegrees = 0.0;
+		Radius = 0.0;
+		return;
+	}
+
 	//Convert from Cartesian to polar coordinates
 	Radius = XYZ.Length();
 	double l = FMath::Atan2(XYZ.Y, XYZ.X);
@@ -981,19 +1025,19 @@ FString UCelestialMaths::GetPreciseVectorString(FVector Vector, int32 MinimumFra
 		);
 }
 
-FString UCelestialMaths::Conv_StarInfoToString(const FStarInfo& StarInfo)
+FString UCelestialMaths::Conv_CelestialBodyToString(const FCelestialBody& CelestialBody) 
 {
-	return StarInfo.ToString();
+	return CelestialBody.ToString();
 }
 
-FString UCelestialMaths::Conv_PlanetaryBodyInfoToString(const FPlanetaryBodyInfo& PlanetaryBodyInfo)
+FString UCelestialMaths::Conv_PlanetaryBodyToString(const FPlanetaryBody& PlanetaryBody)
 {
-	return PlanetaryBodyInfo.ToString();
+	return PlanetaryBody.ToString();
 }
 
-FString UCelestialMaths::Conv_SunInfoToString(const FSunInfo& SunInfo)
+FString UCelestialMaths::Conv_StellarBodyToString(const FStellarBody& StellarBody)
 {
-	return SunInfo.ToString();
+	return StellarBody.ToString();
 }
 
 FString UCelestialMaths::Conv_RightAscensionToString(const double& RightAscensionHours)
@@ -1026,9 +1070,28 @@ FString UCelestialMaths::Conv_DMSToString(bool Sign, int32 Degrees, int32 Minute
 	return FString::Printf(TEXT("%c%d°%02d\'%05.2f\""), Sign? '+' : '-' , Degrees, Minutes, Seconds); 
 }
 
+FString UCelestialMaths::Conv_DegreesToHMSString(double Degrees)
+{
+	int32 H = 0;
+	int32 M = 0;
+	double S = 0.0;
+	DegreesToHMS(Degrees, H, M,S);
+	return Conv_HMSToString(H,M,S); 
+}
+
+FString UCelestialMaths::Conv_DegreesToDMSString(double Degrees)
+{
+	int32 D = 0;
+	int32 M = 0;
+	double S = 0.0;
+	bool Sign = false;
+	DegreesToDMS(Degrees, Sign, D, M, S  );
+	return Conv_DMSToString(Sign, D, M, S);
+}
+
 #pragma endregion
 
-// Private
+#pragma region Private Functions and variables
 
 int32 UCelestialMaths::FloorForJulianDate(double JulianDate)
 {
@@ -1041,4 +1104,51 @@ int32 UCelestialMaths::FloorForJulianDate(double JulianDate)
 		return StaticCast<int32>(JulianDate);
 	}
 	return StaticCast<int32>(FMath::Floor(JulianDate)-1);
-}
\ No newline at end of file
+}
+
+FMatrix UCelestialMaths::GetRotationXMatrix(double AngleRadians)
+{
+	// Compute cosine and sine of the angle
+	double CosTheta = FMath::Cos(AngleRadians);
+	double SinTheta = FMath::Sin(AngleRadians);
+	
+	// Construct the 4x4 rotation matrix manually
+	return FMatrix(
+			FPlane(1.0, 0.0,      0.0,      0.0), // Row 1
+			FPlane(0.0, CosTheta, SinTheta,0.0), // Row 2
+			FPlane(0.0, -SinTheta, CosTheta, 0.0), // Row 3
+			FPlane(0.0, 0.0,      0.0,      1.0)  // Row 4 (required for FMatrix, but ignored here)
+		);
+}
+
+FMatrix UCelestialMaths::GetRotationYMatrix(double AngleRadians)
+{
+	// Compute cosine and sine of the angle
+	double CosTheta = FMath::Cos(AngleRadians);
+	double SinTheta = FMath::Sin(AngleRadians);
+	
+	// Construct the 4x4 rotation matrix manually
+	return FMatrix(
+		FPlane(CosTheta, 0.0, -SinTheta,0.0), // Row 1
+		FPlane(0.0,      1.0, 0.0,     0.0), // Row 2
+		FPlane(SinTheta,0.0, CosTheta,0.0), // Row 3
+		FPlane(0.0,      0.0, 0.0,     1.0)  // Row 4 (required for FMatrix, but ignored here)
+		);
+}
+
+FMatrix UCelestialMaths::GetRotationZMatrix(double AngleRadians)
+{
+	// Compute cosine and sine of the angle
+	double CosTheta = FMath::Cos(AngleRadians);
+	double SinTheta = FMath::Sin(AngleRadians);
+	
+	// Construct the 4x4 rotation matrix manually
+	return FMatrix(
+			FPlane(CosTheta, SinTheta, 0.0, 0.0), // Row 1
+			FPlane(-SinTheta,  CosTheta, 0.0, 0.0), // Row 2
+			FPlane(0.0,      0.0,       1.0, 0.0), // Row 3
+			FPlane(0.0,      0.0,       0.0, 1.0)  // Row 4 (required for FMatrix, but ignored here)
+		);
+}
+
+#pragma endregion
\ No newline at end of file
diff --git a/Engine/Plugins/Experimental/CelestialVault/Source/CelestialVault/Private/CelestialVaultDaySequenceActor.cpp b/Engine/Plugins/Experimental/CelestialVault/Source/CelestialVault/Private/CelestialVaultDaySequenceActor.cpp
--- a/Engine/Plugins/Experimental/CelestialVault/Source/CelestialVault/Private/CelestialVaultDaySequenceActor.cpp
+++ b/Engine/Plugins/Experimental/CelestialVault/Source/CelestialVault/Private/CelestialVaultDaySequenceActor.cpp
@@ -1,9 +1,7 @@
 ﻿// Copyright Epic Games, Inc. All Rights Reserved.
 
-
 #include "CelestialVaultDaySequenceActor.h"
 
-
 // Components
 #include "Components/DirectionalLightComponent.h"
 #include "Components/ExponentialHeightFogComponent.h"
@@ -29,7 +27,23 @@
 #include "CelestialDataTypes.h"
 #include "CelestialVault.h"
 
+#if CELESTIAL_VAULT_ENABLE_DRAW_DEBUG
+
+#include "Engine/Engine.h"
+#include "Debug/DebugDrawService.h"
+#include "SceneView.h"
+#include "Engine/Canvas.h"
 
+// Celestial Vault Global ShowFlag
+TCustomShowFlag<EShowFlagShippingValue::Dynamic> CelestialVaultCustomShowFlag(
+	TEXT("CelestialVault"),
+	false,
+	EShowFlagGroup::SFG_Visualize,
+	NSLOCTEXT("ShowFlagDisplayName", "CelestialVault", "Celestial Vault")
+);
+#endif
+
+#define UE_ONE_KILOMETER 100000.0
 
 // Sets default values
 ACelestialVaultDaySequenceActor::ACelestialVaultDaySequenceActor(const FObjectInitializer& Init)
@@ -51,8 +65,6 @@ ACelestialVaultDaySequenceActor::ACelestialVaultDaySequenceActor(const FObjectIn
 	GlobalPostProcessVolume->Settings.bOverride_LocalExposureBlurredLuminanceBlend = true;
 	GlobalPostProcessVolume->Settings.LocalExposureBlurredLuminanceBlend = 0.4f;
 	
-	
-	
 	// Components Attached to Root
 	SkyLightComponent = CreateDefaultSubobject<USkyLightComponent>(TEXT("Sky Light"));
 	SkyLightComponent->SetupAttachment(RootComponent);
@@ -203,17 +215,99 @@ void ACelestialVaultDaySequenceActor::OnConstruction(const FTransform& Transform
 		}
 
 		// Replace the Moon Material with a MID
-		UMaterialInstanceDynamic * MoonDiscMaterial = Cast<UMaterialInstanceDynamic>(MoonDiscComponent->GetMaterial(0));
-		if (!MoonDiscMaterial)
+		if (MoonDiscComponent)
 		{
-			MoonDiscMaterial = MoonDiscComponent->CreateAndSetMaterialInstanceDynamic(0); // Make the material Dynamic to control the Phase
+			UMaterialInstanceDynamic * MoonDiscMaterial = Cast<UMaterialInstanceDynamic>(MoonDiscComponent->GetMaterial(0));
+			if (!MoonDiscMaterial)
+			{
+				MoonDiscMaterial = MoonDiscComponent->CreateAndSetMaterialInstanceDynamic(0); // Make the material Dynamic to control the Phase
+			}
+			MoonDiscMaterial->SetScalarParameterValue(FName("MoonAge"), MoonAge);	
 		}
-		MoonDiscMaterial->SetScalarParameterValue(FName("MoonAge"), MoonAge);
 		
 		RebuildAll();  
 	}
 }
 
+void ACelestialVaultDaySequenceActor::PostRegisterAllComponents()
+{
+	Super::PostRegisterAllComponents();
+
+#if CELESTIAL_VAULT_ENABLE_DRAW_DEBUG
+	if (!DrawDebugDelegateHandle.IsValid())
+	{
+		DrawDebugDelegateHandle = UDebugDrawService::Register(*CelestialVaultSF, FDebugDrawDelegate::CreateUObject(this, &ACelestialVaultDaySequenceActor::DebugDrawCallback));
+	}
+#endif 
+}
+
+void ACelestialVaultDaySequenceActor::UnregisterAllComponents(bool bForReregister)
+{
+#if CELESTIAL_VAULT_ENABLE_DRAW_DEBUG
+	if (DrawDebugDelegateHandle.IsValid())
+	{
+		UDebugDrawService::Unregister(DrawDebugDelegateHandle);
+		DrawDebugDelegateHandle.Reset();
+	}
+#endif 	
+	Super::UnregisterAllComponents(bForReregister);
+}
+
+#if CELESTIAL_VAULT_ENABLE_DRAW_DEBUG
+
+void ACelestialVaultDaySequenceActor::DebugDrawCallback(UCanvas* Canvas, APlayerController* PC) const
+{
+	// Get the current editor view’s flags (pseudo; adapt to your context)
+	const FEngineShowFlags& Flags = Canvas->SceneView->Family->EngineShowFlags;
+	if (CelestialVaultCustomShowFlag.IsEnabled(Flags))
+	{
+		// Draw in Green
+		Canvas->SetDrawColor(FColor::Green);
+
+		// Build message string before displaying it at once
+		FString DebugMessage;
+		DebugMessage += FString("------ Celestial Vault ------\n");
+		DebugMessage += FString("    -- Time -- \n");
+		
+
+#if WITH_EDITORONLY_DATA
+		FDateTime LocalTime(Year, Month, Day, TimeOfDayPreview.Hours, TimeOfDayPreview.Minutes, TimeOfDayPreview.Seconds);
+#else
+		FDateTime LocalTime(Year, Month, Day, 12, 0, 0);
+#endif
+		FDateTime UTCTime = UCelestialMaths::LocalTimeToUTCTime(LocalTime, GMT_TimeZone, bIsDaylightSavings);
+		double JulianDate = UCelestialMaths::UTCDateTimeToJulianDate(UTCTime);
+		double JulianCenturies = UCelestialMaths::JulianDateToJulianCenturies(JulianDate); 
+		double GMST = UCelestialMaths::DateTimeToGreenwichMeanSiderealTime(UTCTime);
+		double LMST = UCelestialMaths::LocalSideralTime(Longitude, GMST); // TODO /!\ ECEF
+
+		 
+		DebugMessage += FString::Printf(TEXT("    Local Time: %s"), *LocalTime.ToFormattedString(TEXT("%d %b %Y, %H:%M:%S"))) + LINE_TERMINATOR;
+		DebugMessage += FString::Printf(TEXT("    UTC Time: %s"), *UTCTime.ToFormattedString(TEXT("%d %b %Y, %H:%M:%S"))) + LINE_TERMINATOR;
+		DebugMessage += FString::Printf(TEXT("    Julian Date: %.5f"), JulianDate) + LINE_TERMINATOR;
+		DebugMessage += FString::Printf(TEXT("    Julian Centuries: %.5f"), JulianCenturies) + LINE_TERMINATOR;
+
+		DebugMessage += FString::Printf(TEXT("    GMST at 0:00: %s (%.5f)°"), *UCelestialMaths::Conv_DegreesToHMSString(GMST_AtTOD_0), GMST_AtTOD_0) + LINE_TERMINATOR;
+		DebugMessage += FString::Printf(TEXT("    GMST: %s (%.5f)°"), *UCelestialMaths::Conv_DegreesToHMSString(GMST), GMST) + LINE_TERMINATOR;
+		DebugMessage += FString::Printf(TEXT("    LMST: %s (%.5f)°"), *UCelestialMaths::Conv_DegreesToHMSString(LMST), LMST) + LINE_TERMINATOR;
+
+		// Sun and Moon Info
+		DebugMessage += LINE_TERMINATOR;
+		DebugMessage += FString("    -- Sun -- \n");
+		FStellarBody SunAtTime = GetSun(JulianDate);
+		DebugMessage += SunAtTime.ToString() + LINE_TERMINATOR;
+
+		DebugMessage += FString("    -- Moon -- \n");
+		FPlanetaryBody MoonAtTime = GetMoon(JulianDate);
+		DebugMessage += Moon.ToString() + LINE_TERMINATOR;
+		
+		// Draw
+		Canvas->DrawText(GetDefault<UEngine>()->GetMediumFont(), DebugMessage, 10, 10);
+	}
+}
+
+#endif
+
 //////// Public Functions
 FDateTime ACelestialVaultDaySequenceActor::GetDate() const
 {
@@ -233,41 +327,40 @@ FDateTime ACelestialVaultDaySequenceActor::GetDate() const
 	return FDateTime(Year, Month, Day, 0,0,0);
 }
 
-FSunInfo ACelestialVaultDaySequenceActor::GetSunInfo(double JulianDate)
+FStellarBody ACelestialVaultDaySequenceActor::GetSun(double JulianDate) const
 {
-	FSunInfo ResultSunInfo = UCelestialMaths::GetSunInformation(JulianDate, Latitude, Longitude);
+	FStellarBody ResultSunInfo = UCelestialMaths::GetSunInformation(JulianDate, Latitude, Longitude);
 
 	// Location
-	FVector SunLocation = UCelestialMaths::RADECToXYZ_RH(ResultSunInfo.RA * 15.0, ResultSunInfo.DEC, 1000.0); // TODO - Add the Sun Distance To SunInfo 
+	FVector SunLocation = UCelestialMaths::RADECToXYZ_RH(ResultSunInfo.RA * 15.0, ResultSunInfo.DEC, 1000.0); // TODO - Add the Sun Distance To Sun 
 	SunLocation.Y *= -1.0; // Convert to UE Frame by inverting Y
-	ResultSunInfo.UETransform.SetLocation(SunLocation);
+	ResultSunInfo.UELocalTransform.SetLocation(SunLocation);
 	ResultSunInfo.DirectionTowardEarth = (FVector::ZeroVector - SunLocation).GetSafeNormal();
 	return ResultSunInfo;
 }
 
-FPlanetaryBodyInfo ACelestialVaultDaySequenceActor::GetMoonInfo(double JulianDate)
+FPlanetaryBody ACelestialVaultDaySequenceActor::GetMoon(double JulianDate) const
 {
-	double UEDistance = CelestialVaultDistance * 100000.0 * MoonVaultPercentage / 100.0;
-	FPlanetaryBodyInfo ResultMoonInfo = GetPlanetaryBodyInfo(FPlanetaryBodyInputData::Moon, JulianDate, UEDistance, MoonScale);
+	double UEDistance = CelestialVaultDistance * UE_ONE_KILOMETER * MoonVaultPercentage / 100.0; 
+	FPlanetaryBody ResultMoonInfo = GetPlanetaryBodyInfo(FPlanetaryBodyInputData::Moon, JulianDate);
 
-	if (bManualControl)
+	// Set the Moon specifics, and compute the Transform
+	ResultMoonInfo.Age = UCelestialMaths::GetMoonNormalizedAgeSimple(JulianDate);;
+	ResultMoonInfo.IlluminationPercentage = UCelestialMaths::GetIlluminationPercentage(ResultMoonInfo.Age);
+	ResultMoonInfo.ComputeUELocalTransform();
+	
+	if (bManualControl) // Override 
 	{
 		// Phase/Age
 		ResultMoonInfo.Age = MoonAge;
 		ResultMoonInfo.IlluminationPercentage = UCelestialMaths::GetIlluminationPercentage(MoonAge);
 
 		// Location
-		FSunInfo SunInfoTemp = GetSunInfo(JulianDate);
+		FStellarBody SunInfoTemp = GetSun(JulianDate);
 		ResultMoonInfo.RA = UCelestialMaths::ModPositive(SunInfoTemp.RA + MoonOffset_RA, 24.0);
-		ResultMoonInfo.DEC = UCelestialMaths::ModPositive(SunInfoTemp.DEC + MoonOffset_DEC, 360.0);
-		
+		ResultMoonInfo.DEC = UCelestialMaths::ModPositive(SunInfoTemp.DEC + MoonOffset_DEC + 180.0, 360.0) - 180.0;
 		// We changed the RA, so we need to update the transform
-		ResultMoonInfo.ComputeTranform(UEDistance, MoonScale);
-	}
-	else
-	{
-		ResultMoonInfo.Age = UCelestialMaths::GetMoonNormalizedAgeSimple(JulianDate);;
-		ResultMoonInfo.IlluminationPercentage = UCelestialMaths::GetIlluminationPercentage(ResultMoonInfo.Age);
+		ResultMoonInfo.ComputeUELocalTransform(MoonScale, UEDistance);
 	}
 	return ResultMoonInfo;
 }
@@ -289,7 +382,7 @@ void ACelestialVaultDaySequenceActor::SetMoonDiscAge(float InMoonAge)
 	}
 }
 
-bool ACelestialVaultDaySequenceActor::GetClosestStarInfo(FVector ObserverLocation, FVector LookupDirection, double ThresholdAngleDegrees, FStarInfo& FoundStarInfo, FTransform& StarTransform)
+bool ACelestialVaultDaySequenceActor::GetClosestStar(FVector ObserverLocation, FVector LookupDirection, double ThresholdAngleDegrees, FStellarBody& FoundStarInfo, FTransform& StarTransform)
 {
 	StarTransform = FTransform::Identity;
 	
@@ -302,41 +395,36 @@ bool ACelestialVaultDaySequenceActor::GetClosestStarInfo(FVector ObserverLocatio
 	double CosThresholdAngle = FMath::Cos(FMath::DegreesToRadians(ThresholdAngleDegrees));
 	LookupDirection.Normalize();
 	double MinCos = -1.0;
-	int32 ClosestInstanceIndex = -1;
+	bool StarFound = false;
 	
-	for (int StarInfoIndex = 0; StarInfoIndex < StarsInfo.Num(); StarInfoIndex++)
+	FTransform StarsComponentTransform = StarsComponent->GetComponentTransform();
+	for (FStellarBody StarInfo : Stars)
 	{
-		FTransform ISMInstanceTransform;
-		StarsComponent->GetInstanceTransform(StarsInfo[StarInfoIndex].ISMInstanceIndex, ISMInstanceTransform, true );
-
-		FVector DirectionToInstance = ISMInstanceTransform.GetLocation() - ObserverLocation;
-		DirectionToInstance.Normalize();
+		// We need to compute the world Transform for the body, because the celestial vault has rotated and the Body UE Transform is local to the Sky. 
+		FTransform WorldTransform =  StarInfo.UELocalTransform * StarsComponentTransform;
+			
+		FVector DirectionToInstance = WorldTransform.GetLocation() - ObserverLocation;
 
 		double CosDeltaAngle = FVector::DotProduct(LookupDirection, DirectionToInstance);
 		if (CosDeltaAngle > CosThresholdAngle && CosDeltaAngle > MinCos)
 		{
 			// Inside the cone angle, and closer to latest
 			MinCos = CosDeltaAngle;
-			ClosestInstanceIndex = StarInfoIndex;
+			FoundStarInfo = StarInfo;
+			StarTransform = WorldTransform;
+			StarFound = true;
 		}
 	}
 
-	if (ClosestInstanceIndex != -1)
-	{
-		StarsComponent->GetInstanceTransform(StarsInfo[ClosestInstanceIndex].ISMInstanceIndex, StarTransform, true ); 
-		FoundStarInfo = StarsInfo[ClosestInstanceIndex];
-		return true;
-	}
-
-	return false;
+	return StarFound;
 }
 
-bool ACelestialVaultDaySequenceActor::GetClosestPlanetaryBody(FVector ObserverPosition, FVector LookupDirection, double ThresholdAngleDegrees, FPlanetaryBodyInfo& FoundPlanetaryBodyInfo, FTransform& BodyTransform)
+bool ACelestialVaultDaySequenceActor::GetClosestPlanetaryBody(FVector ObserverLocation, FVector LookupDirection, double ThresholdAngleDegrees, FPlanetaryBody& FoundPlanetaryBodyInfo, FTransform& BodyTransform)
 {
 	BodyTransform = FTransform::Identity;
 	
-	// We query only if we generated the StarInfo  
-	if (!bKeepPlanetsInfos)
+	// We query only if we generated the Planet Info  
+	if (!bKeepPlanetsInfos && PlanetsComponent)
 	{
 		return false;
 	}
@@ -344,89 +432,76 @@ bool ACelestialVaultDaySequenceActor::GetClosestPlanetaryBody(FVector ObserverPo
 	double CosThresholdAngle = FMath::Cos(FMath::DegreesToRadians(ThresholdAngleDegrees));
 	LookupDirection.Normalize();
 	double MinCos = -1.0;
-	int32 ClosestInstanceIndex = -1;
-
-	for (FPlanetaryBodyInfo BodyInfo : PlanetsInfos)
-	{
-		if (PlanetsComponent)
-		{
-			// We need to use the ISM component and query the world transform because the celestial vault has rotated
-			FTransform ISMInstanceTransform;
-			PlanetsComponent->GetInstanceTransform(BodyInfo.ISMInstanceIndex, ISMInstanceTransform, true );
-
-			FVector DirectionToInstance = ISMInstanceTransform.GetLocation() - ObserverPosition;
-			DirectionToInstance.Normalize();
-
-			double CosDeltaAngle = FVector::DotProduct(LookupDirection, DirectionToInstance);
-			if (CosDeltaAngle > CosThresholdAngle && CosDeltaAngle > MinCos)
-			{
-				// Inside the cone angle, and closer to latest
-				MinCos = CosDeltaAngle;
-				ClosestInstanceIndex = BodyInfo.ISMInstanceIndex;
-			}
-		}
-	}
-
-	// Check for the moon
-	bool bIsMoonCloser = false;
+	bool BodyFound = false;
+	
+	// Check for the moon first
 	if (MoonDiscComponent)
 	{
 		// We need to use the ISM component and query the world transform because the celestial vault has rotated
 		FVector MoonLocation = MoonDiscComponent->GetComponentLocation();
-		FVector DirectionToMoon = MoonLocation - ObserverPosition;
+		FVector DirectionToMoon = MoonLocation - ObserverLocation;
 		DirectionToMoon.Normalize();
 
 		double CosDeltaAngle = FVector::DotProduct(LookupDirection, DirectionToMoon);
 		if (CosDeltaAngle > CosThresholdAngle && CosDeltaAngle > MinCos)
 		{
-			// Inside the cone angle, and closer to latest
-			bIsMoonCloser = true;
+			// Inside the cone angle. At least we have found the moon. 
+			MinCos = CosDeltaAngle;
+			BodyTransform = MoonDiscComponent->GetComponentTransform();
+			FoundPlanetaryBodyInfo = Moon;
+			BodyFound = true;
 		}
 	}
 
-	if (bIsMoonCloser)
+	// Check for planets now
+	FTransform PlanetComponentTransform = PlanetsComponent->GetComponentTransform();
+	for (FPlanetaryBody BodyInfo : Planets)
 	{
-		BodyTransform = MoonDiscComponent->GetComponentTransform();
-		FoundPlanetaryBodyInfo = MoonBodyInfo;
-		return true;
-	}
-	else if (ClosestInstanceIndex != -1)
-	{
-		PlanetsComponent->GetInstanceTransform(ClosestInstanceIndex, BodyTransform, true );
-		FoundPlanetaryBodyInfo = PlanetsInfos[ClosestInstanceIndex];
-		return true;
+		if (PlanetsComponent)
+		{
+			// We need to compute the world Transform for the body, because the celestial vault has rotated and the Body UE Transform is local to the Sky. 
+			FTransform WorldTransform =  BodyInfo.UELocalTransform * PlanetComponentTransform;
+			
+			FVector DirectionToInstance = WorldTransform.GetLocation() - ObserverLocation;
+			DirectionToInstance.Normalize();
+
+			double CosDeltaAngle = FVector::DotProduct(LookupDirection, DirectionToInstance);
+			if (CosDeltaAngle > CosThresholdAngle && CosDeltaAngle > MinCos)
+			{
+				// Inside the cone angle, and closer to latest found body
+				MinCos = CosDeltaAngle;
+				FoundPlanetaryBodyInfo = BodyInfo;
+				BodyTransform = WorldTransform;
+				BodyFound = true;
+			}
+		}
 	}
 
-	return false;
+	return BodyFound;
 }
 
-bool ACelestialVaultDaySequenceActor::GetPlanetaryBodyByOrbitType(EOrbitType OrbitType, FPlanetaryBodyInfo& FoundPlanetaryBodyInfo, FTransform& BodyTransform)
+bool ACelestialVaultDaySequenceActor::GetPlanetaryBodyByVSOP87Type(EVSOP87BodyType VSOP87Type, FPlanetaryBody& FoundPlanetaryBodyInfo, FTransform& BodyTransform)
 {
 	BodyTransform = FTransform::Identity;
 	
-	// We query only if we generated the StarInfo  
+	// We query only if we generated the Planets Info  
 	if (!bKeepPlanetsInfos)
 	{
 		return false;
 	}
 	
-	if (OrbitType == EOrbitType::Moon && MoonDiscComponent) 
+	if (VSOP87Type == EVSOP87BodyType::Moon && MoonDiscComponent) 
 	{
 		BodyTransform = MoonDiscComponent->GetComponentTransform();
-		FoundPlanetaryBodyInfo = MoonBodyInfo;
+		FoundPlanetaryBodyInfo = Moon;
 		return true;
 	}
 	else
 	{
-		for (FPlanetaryBodyInfo PlanetaryBodyInfo : PlanetsInfos)
+		for (FPlanetaryBody PlanetaryBodyInfo : Planets)
 		{
-			if (PlanetaryBodyInfo.OrbitType == OrbitType)
+			if (PlanetaryBodyInfo.VSOP87BodyType == VSOP87Type)
 			{
-				if (PlanetsComponent)
-				{
-					PlanetsComponent->GetInstanceTransform(PlanetaryBodyInfo.ISMInstanceIndex, BodyTransform, true );
-				}
-				
 				FoundPlanetaryBodyInfo = PlanetaryBodyInfo;
 				return true;
 			}
@@ -437,12 +512,7 @@ bool ACelestialVaultDaySequenceActor::GetPlanetaryBodyByOrbitType(EOrbitType Orb
 
 double ACelestialVaultDaySequenceActor::GetDayCelestialVaultAngle() const
 {
-	double TimeOffset = GMT_TimeZone;
-	if (bIsDaylightSavings)
-	{
-		TimeOffset += 1.0;
-	}
-	return GMST_AtTOD_0 - TimeOffset * 15.0; 
+	return GMST_AtTOD_0; 
 }
 
 #if WITH_EDITOR
@@ -485,8 +555,9 @@ void ACelestialVaultDaySequenceActor::PostEditChangeProperty(FPropertyChangedEve
 		PropertyName == GET_MEMBER_NAME_CHECKED(ACelestialVaultDaySequenceActor, Longitude))
 	{
 		// DateTime has changed, we need to update the Reference CelestialVault Angle.
-		FDateTime DateTime = GetDate();
-		GMST_AtTOD_0 = UCelestialMaths::DateTimeToGreenwichMeanSiderealTime(DateTime);
+		FDateTime LocalDateTimeZero = GetDate();
+		FDateTime UTCDateTimeZero = UCelestialMaths::LocalTimeToUTCTime(LocalDateTimeZero, GMT_TimeZone, bIsDaylightSavings);
+		GMST_AtTOD_0 = UCelestialMaths::DateTimeToGreenwichMeanSiderealTime(UTCDateTimeZero);
 		double CelestialVaultAngle = GetDayCelestialVaultAngle();
 		CelestialVaultComponent->SetRelativeRotation( FRotator(0.0, CelestialVaultAngle, 0.0) );
 
@@ -550,7 +621,7 @@ void ACelestialVaultDaySequenceActor::InitStars()
 	{
 		StarsComponent->ClearInstances();
 		StarsComponent->SetNumCustomDataFloats(4);
-		StarsInfo.Empty();
+		Stars.Empty();
 		
 		// Check for the Catalog
 		if (!CelestialStarCatalog && !FictionalStarCatalog)
@@ -592,10 +663,10 @@ void ACelestialVaultDaySequenceActor::InitStars()
 							// Keep trace of the Star Information for further runtime queries					
 							if (bKeepStarsInfo)
 							{
-								FStarInfo StarInfo;
+								FStellarBody StarInfo;
 								StarInfo.RA = CelestialStarInputData.RA;
 								StarInfo.DEC = CelestialStarInputData.DEC;
-								StarInfo.DistanceInPC = CelestialStarInputData.DistanceInPC;
+								StarInfo.DistanceInAU = UCelestialMaths::ParsecsToAstronomicalUnits(CelestialStarInputData.DistanceInPC);
 								StarInfo.Name = CelestialStarInputData.Name;
 								StarInfo.Magnitude = CelestialStarInputData.Magnitude;
 								StarInfo.Color = StarColor;
@@ -603,9 +674,11 @@ void ACelestialVaultDaySequenceActor::InitStars()
 								StarInfo.HenryDraperID = CelestialStarInputData.HenryDraperID;
 								StarInfo.YaleBrightStarID = CelestialStarInputData.YaleBrightStarID;
 								StarInfo.ColorIndex = CelestialStarInputData.ColorIndex;
-								StarInfo.ISMInstanceIndex = NewIndex;
+
+								double UEBodyDistance = CelestialVaultDistance * UE_ONE_KILOMETER * StarsVaultPercentage / 100.0;
+								StarInfo.ComputeUELocalTransform(1.0f, UEBodyDistance);
 								// Maybe add other computed values here... 
-								StarsInfo.Add(StarInfo);
+								Stars.Add(StarInfo);
 							}
 						}
 					}
@@ -645,15 +718,15 @@ void ACelestialVaultDaySequenceActor::InitStars()
 							// Keep trace of the Star Information for further runtime queries					
 							if (bKeepStarsInfo)
 							{
-								FStarInfo StarInfo;
+								FStellarBody StarInfo;
 								StarInfo.RA = StarInputData.RA;
 								StarInfo.DEC = StarInputData.DEC;
-								StarInfo.DistanceInPC = StarInputData.DistanceInPC;
+								StarInfo.DistanceInAU = UCelestialMaths::ParsecsToAstronomicalUnits(StarInputData.DistanceInPC);
 								StarInfo.Name = StarInputData.Name;
 								StarInfo.Magnitude = StarInputData.Magnitude;
 								StarInfo.Color = StarColor;
-								StarInfo.ISMInstanceIndex = NewIndex;
-								StarsInfo.Add(StarInfo);
+								StarInfo.ComputeUELocalTransform();
+								Stars.Add(StarInfo);
 							}
 						}
 					}
@@ -676,16 +749,16 @@ void ACelestialVaultDaySequenceActor::InitPlanetaryBodies()
 	double JulianDay = UCelestialMaths::UTCDateTimeToJulianDate(UTCTime);
 
 	// Init the sun
-	SunInfo = GetSunInfo(JulianDay);
+	Sun = GetSun(JulianDay);
 	
 	// Init the Moon
 	if (MoonDiscComponent)
 	{
 		// We don't want to use the GetMoonTransform function here because we also want to keep the Moon information for the day.
 		// It has a pitfall because the Moon information will not be adjusted to the current TOD, but that doesn't prevent anyone for doing a manuel query at some point
-		MoonBodyInfo = GetMoonInfo(JulianDay);
-		SetMoonDiscAge(MoonBodyInfo.Age);
-		MoonDiscComponent->SetRelativeTransform(MoonBodyInfo.UETransform);
+		Moon = GetMoon(JulianDay);
+		SetMoonDiscAge(Moon.Age);
+		MoonDiscComponent->SetRelativeTransform(Moon.UELocalTransform);
 	}
 	
 	// Init the Planets from the catalog
@@ -693,7 +766,7 @@ void ACelestialVaultDaySequenceActor::InitPlanetaryBodies()
 	{
 		PlanetsComponent->ClearInstances();
 		PlanetsComponent->SetNumCustomDataFloats(2);
-		PlanetsInfos.Empty();
+		Planets.Empty();
 
 		// Safety check on Catalog Data
 		if (!PlanetsCatalog)
@@ -712,23 +785,24 @@ void ACelestialVaultDaySequenceActor::InitPlanetaryBodies()
 		// Init from the DataTable
 		PlanetsCatalog->ForeachRow<FPlanetaryBodyInputData>("AEarthDaySequenceActor::InitPlanetaryBodies", [&](const FName& Key, const FPlanetaryBodyInputData& InputPlanetaryBody)
 			{
-				double UEBodyDistance = CelestialVaultDistance * 100000.0 * PlanetsVaultPercentage / 100.0;
-				FPlanetaryBodyInfo BodyInfo = GetPlanetaryBodyInfo(InputPlanetaryBody, JulianDay, UEBodyDistance, PlanetsScale);
-
+				double UEBodyDistance = CelestialVaultDistance * UE_ONE_KILOMETER * PlanetsVaultPercentage / 100.0;
+				FPlanetaryBody BodyInfo = GetPlanetaryBodyInfo(InputPlanetaryBody, JulianDay);
+				// Compute the Transform
+				BodyInfo.ComputeUELocalTransform(PlanetsScale, UEBodyDistance);
+			
 				// Add the new ISM Instance
-				int32 NewIndex = PlanetsComponent->AddInstance(BodyInfo.UETransform);
-				BodyInfo.ISMInstanceIndex = NewIndex;
+				int32 NewIndex = PlanetsComponent->AddInstance(BodyInfo.UELocalTransform);
 			
 				// Add the Custom data (column index to sample the planets atlas texture)
 				TArray<float> NewCustomData;
 				NewCustomData.Add(InputPlanetaryBody.TextureColumnIndex);
-				NewCustomData.Add(BodyInfo.ApparentMagnitude);
+				NewCustomData.Add(BodyInfo.Magnitude);
 				PlanetsComponent->SetCustomData(NewIndex, NewCustomData);
 
 				// Keep trace of the PlanetaryBody for further queries (the Datatable is readonly, so store in another object)
 				if (bKeepPlanetsInfos)
 				{
-					PlanetsInfos.Add(BodyInfo);	
+					Planets.Add(BodyInfo);	
 				}
 			}
 		);
@@ -736,33 +810,30 @@ void ACelestialVaultDaySequenceActor::InitPlanetaryBodies()
 	}
 }
 
-FPlanetaryBodyInfo ACelestialVaultDaySequenceActor::GetPlanetaryBodyInfo(const FPlanetaryBodyInputData& InputPlanetaryBody, double JulianDay, double UEDistance, double BodyScale) const
+FPlanetaryBody ACelestialVaultDaySequenceActor::GetPlanetaryBodyInfo(const FPlanetaryBodyInputData& InputPlanetaryBody, double JulianDay) const
 {
-	FPlanetaryBodyInfo BodyInfo;
+	FPlanetaryBody BodyInfo;
 	BodyInfo.OrbitType = InputPlanetaryBody.OrbitType;
 	BodyInfo.Name = InputPlanetaryBody.Name;
 	BodyInfo.Radius = InputPlanetaryBody.Radius;
 
 	// Compute location	
-	double RAHours, DECDegrees, DistanceToEarthAU, DistanceToSunAU, DistanceEarthToSunAU;
-	UCelestialMaths::GetBodyCelestialCoordinatesAU(JulianDay, InputPlanetaryBody, Latitude, Longitude, RAHours, DECDegrees, DistanceToEarthAU, DistanceToSunAU, DistanceEarthToSunAU);
+	double RAJ2000Hours, DECJ2000Degrees, RAHours, DECDegrees, DistanceToEarthAU, DistanceToSunAU, DistanceEarthToSunAU;
+	
+	UCelestialMaths::GetBodyCelestialCoordinatesAU(JulianDay, InputPlanetaryBody.VSOP87BodyType, Latitude, Longitude, RAJ2000Hours, DECJ2000Degrees,RAHours, DECDegrees, DistanceToEarthAU, DistanceToSunAU, DistanceEarthToSunAU);
+	BodyInfo.RAJ2000 = RAJ2000Hours;
+	BodyInfo.DECJ2000 = DECJ2000Degrees;
 	BodyInfo.RA = RAHours;
 	BodyInfo.DEC = DECDegrees;
 	BodyInfo.DistanceInAU = DistanceToEarthAU;
 
 	// Compute Magnitude
 	double Phase; // TODO - Check the PhaseAngle Computations. 
-	BodyInfo.ApparentMagnitude = UCelestialMaths::GetPlanetaryBodyMagnitude(InputPlanetaryBody, DistanceToEarthAU, DistanceToSunAU, DistanceEarthToSunAU, Phase);
+	BodyInfo.Magnitude = UCelestialMaths::GetPlanetaryBodyMagnitude(InputPlanetaryBody.VSOP87BodyType, DistanceToEarthAU, DistanceToSunAU, DistanceEarthToSunAU, Phase);
 	BodyInfo.Age = Phase;
 	
 	// Compute the True and the Scaled apparent Diameters
 	BodyInfo.ApparentDiameterDegrees = FMath::RadiansToDegrees(FMath::Atan2(InputPlanetaryBody.Radius * 1000.0, UCelestialMaths::AstronomicalUnitsToMeters(DistanceToEarthAU)) * 2.0);
-	BodyInfo.ScaledApparentDiameterDegrees = FMath::RadiansToDegrees(FMath::Atan2(InputPlanetaryBody.Radius * 1000.0 * BodyScale, UCelestialMaths::AstronomicalUnitsToMeters(DistanceToEarthAU)) * 2.0);
-
-
-	//
-	BodyInfo.ComputeTranform(UEDistance, BodyScale);
-	
 	return BodyInfo;
 }
 
@@ -774,28 +845,12 @@ void ACelestialVaultDaySequenceActor::RebuildAll()
 		SetActorTransform(FTransform::Identity);
 		
 		// DateTime has changed, we need to update the Reference CelestialVault Angle.
-		FDateTime DateTime = GetDate();
-		GMST_AtTOD_0 = UCelestialMaths::DateTimeToGreenwichMeanSiderealTime(DateTime);
-
-		// Get the transformation in proper Celestial World units (meters, right handed)
-		FTransform ECEFFrameToWorldFrame = UCelestialMaths::GetPlanetCenterTransform(Latitude, Longitude, 0);
-
-		// UE Frame are expressed in Left-handed coordinates, and units are in meters - Convert to UE Transform 
-		FMatrix WorldFrameToUEFrame = FMatrix( 
-			FVector(1.0, 0.0, 0.0),		// Easting (X) is UE World X
-			FVector(0.0, -1.0, 0.0),	// Northing (Y) is UE World -Y because of left-handed convention
-			FVector(0.0, 0.0, 1.0),		// Up (Z) is UE World Z 
-			FVector(0.0, 0.0, 0.0));	// No Origin offset
-		FMatrix UEFrameToWorldFrame = WorldFrameToUEFrame.Inverse();
-
-		// Update the rotation part
-		FMatrix TransformMatrix = UEFrameToWorldFrame * ECEFFrameToWorldFrame.ToMatrixNoScale() * WorldFrameToUEFrame;
-		// Get Origin, and convert UE units to meters
-		FVector UEOrigin = TransformMatrix.GetOrigin() * FVector(100.0, 100.0, 100.0);
-		TransformMatrix.SetOrigin(UEOrigin);
-
-		// Apply the transform
-		PlanetCenterTransform = FTransform(TransformMatrix);
+		FDateTime LocalDateTimeZero = GetDate();
+		FDateTime UTCDateTimeZero = UCelestialMaths::LocalTimeToUTCTime(LocalDateTimeZero, GMT_TimeZone, bIsDaylightSavings);
+		GMST_AtTOD_0 = UCelestialMaths::DateTimeToGreenwichMeanSiderealTime(UTCDateTimeZero);
+
+		// Get the transformation in proper UE Frame units, and set it to the Component 
+		PlanetCenterTransform = UCelestialMaths::GetEarthCenterTransformUEFrame(Latitude, Longitude, 0);
 		PlanetCenterComponent->SetWorldTransform(PlanetCenterTransform);
 
 		// Rebuild Sky and Sequence
diff --git a/Engine/Plugins/Experimental/CelestialVault/Source/CelestialVault/Private/CelestialVaultSequence.cpp b/Engine/Plugins/Experimental/CelestialVault/Source/CelestialVault/Private/CelestialVaultSequence.cpp
--- a/Engine/Plugins/Experimental/CelestialVault/Source/CelestialVault/Private/CelestialVaultSequence.cpp
+++ b/Engine/Plugins/Experimental/CelestialVault/Source/CelestialVault/Private/CelestialVaultSequence.cpp
@@ -18,7 +18,7 @@ void FCelestialVaultSequence::BuildSequence(UProceduralDaySequenceBuilder* InBui
 	ACelestialVaultDaySequenceActor* EarthDaySequenceActor = Cast<ACelestialVaultDaySequenceActor>(WeakTargetActor.Get());
 	if (!EarthDaySequenceActor)
 	{
-		UE_LOG(LogCelestialVault, Warning, TEXT("This ProceduralDaySequence applies only on AEarthDaySequenceActors"));
+		UE_LOG(LogCelestialVault, Warning, TEXT("This ProceduralDaySequence applies only on ACelestialVaultDaySequenceActor"));
 		return;
 	}
 
@@ -31,7 +31,7 @@ void FCelestialVaultSequence::BuildSequence(UProceduralDaySequenceBuilder* InBui
 	{
 		InBuilder->SetActiveBoundObject(CelestialVaultComponent);
 		
-		const double NormalizedTimeIncrement = 1.0 / FMath::Max(KeyCount - 1, static_cast<unsigned>(1));
+		const double NormalizedTimeIncrement = 1.0 / FMath::Max(static_cast<signed>(KeyCount) - 1, 1);
 		for (unsigned int Key = 0; Key < KeyCount; ++Key)
 		{
 			double CelestialVaultAngleAtMidnight = EarthDaySequenceActor->GetDayCelestialVaultAngle(); 
@@ -47,7 +47,7 @@ void FCelestialVaultSequence::BuildSequence(UProceduralDaySequenceBuilder* InBui
 
 	if (MoonDiscComponent && MoonLightComponent)
 	{
-		const double NormalizedTimeIncrement = 1.0 / FMath::Max(KeyCount - 1, static_cast<unsigned>(1));
+		const double NormalizedTimeIncrement = 1.0 / FMath::Max(static_cast<signed>(KeyCount) - 1, 1);
 		for (unsigned int Key = 0; Key < KeyCount; ++Key)
 		{
 			const double KeyTime = Key * NormalizedTimeIncrement;
@@ -57,7 +57,7 @@ void FCelestialVaultSequence::BuildSequence(UProceduralDaySequenceBuilder* InBui
 				EarthDaySequenceActor->bIsDaylightSavings);
 			double JulianDay = UCelestialMaths::UTCDateTimeToJulianDate(UTCTime);
 
-			FPlanetaryBodyInfo MoonInfo = EarthDaySequenceActor->GetMoonInfo(JulianDay);
+			FPlanetaryBody  MoonInfo = EarthDaySequenceActor->GetMoon(JulianDay);
 
 			if (Key == 0)
 			{
@@ -67,7 +67,7 @@ void FCelestialVaultSequence::BuildSequence(UProceduralDaySequenceBuilder* InBui
 
 			// Animate Disc
 			InBuilder->SetActiveBoundObject(MoonDiscComponent);
-			InBuilder->AddTransformKey(KeyTime, MoonInfo.UETransform);
+			InBuilder->AddTransformKey(KeyTime, MoonInfo.UELocalTransform);
 
 			// Animate Light
 			InBuilder->SetActiveBoundObject(MoonLightComponent);
@@ -79,8 +79,8 @@ void FCelestialVaultSequence::BuildSequence(UProceduralDaySequenceBuilder* InBui
 	if (UDirectionalLightComponent* SunLightComponent = EarthDaySequenceActor->SunLightComponent)
 	{
 		InBuilder->SetActiveBoundObject(SunLightComponent);
-
-		const double NormalizedTimeIncrement = 1.0 / FMath::Max(KeyCount - 1, static_cast<unsigned>(1));
+		
+		const double NormalizedTimeIncrement = 1.0 / FMath::Max(static_cast<signed>(KeyCount) - 1, 1);
 		for (unsigned int Key = 0; Key < KeyCount; ++Key)
 		{
 			const double KeyTime = Key * NormalizedTimeIncrement;
@@ -90,7 +90,7 @@ void FCelestialVaultSequence::BuildSequence(UProceduralDaySequenceBuilder* InBui
 				EarthDaySequenceActor->bIsDaylightSavings);
 			double JulianDay = UCelestialMaths::UTCDateTimeToJulianDate(UTCTime);
 
-			FSunInfo SunInfo = EarthDaySequenceActor->GetSunInfo(JulianDay);
+			FStellarBody SunInfo = EarthDaySequenceActor->GetSun(JulianDay);
 			// Animate Light
 			InBuilder->AddRotationKey(KeyTime, FRotationMatrix::MakeFromX(SunInfo.DirectionTowardEarth).Rotator());
 		}
