* Refactored the Celestial Data Types for better consistency

diff --git a/Engine/Plugins/Experimental/CelestialVault/Source/CelestialVault/Public/CelestialDataTypes.h b/Engine/Plugins/Experimental/CelestialVault/Source/CelestialVault/Public/CelestialDataTypes.h
--- a/Engine/Plugins/Experimental/CelestialVault/Source/CelestialVault/Public/CelestialDataTypes.h
+++ b/Engine/Plugins/Experimental/CelestialVault/Source/CelestialVault/Public/CelestialDataTypes.h
@@ -3,18 +3,22 @@
 #pragma once
 
 #include "CoreMinimal.h"
-#include "Engine/DataTable.h"
 #include "CelestialDataTypes.generated.h"
 
-
-/** 
- * Any celestial body with an Elliptic Orbit type will have to provide the elliptic parameters
- * Other will use the solar system VSOP87 computations for their location
- */
+/** Different type of orbital motion for the Celestial bodies. Will be used to branch the computational logics */
 UENUM(BlueprintType) 
 enum class EOrbitType : uint8
 {
-	Elliptic = 0, 
+	Star,		// Far away star - No orbit
+	Elliptic,	// Generic Ellipsoidal Orbit
+	VSOP87		// Orbit defined by the VSOP87 Computations
+};
+
+/** Different type Solar System bodies in the VSOP87 Equations */
+UENUM(BlueprintType) 
+enum class EVSOP87BodyType : uint8
+{
+	Sun,
 	Mercury,
 	Venus,
 	Earth,
@@ -27,69 +31,36 @@ enum class EOrbitType : uint8
 };
 
 /**
- * TableRow Base type to describe the Creation parameters of a planetary body (Planet, Moon)
- * Will only be used at creation time, from a proper Data Table
+ * Runtime structure to store the properties common to all Celestial Bodies
  */
-USTRUCT(BlueprintType)
-struct FPlanetaryBodyInputData : public FTableRowBase
+USTRUCT(BlueprintInternalUseOnly)
+struct FCelestialBody
 {
 	GENERATED_BODY()
+	virtual ~FCelestialBody() = default;   // to remove Warning C4265
 
-	FPlanetaryBodyInputData(const FString& InName, EOrbitType InOrbitType, double InRadius) : Name(InName), OrbitType(InOrbitType), Radius(InRadius) {}
-	FPlanetaryBodyInputData() {}
-public:
-	/** The body Name */
+public: 
+	/** The Body Name */
 	UPROPERTY(BlueprintReadWrite, EditAnywhere, Category= "Celestial Vault")
 	FString Name = "";
 
-	/** One of the predefined Orbit types (Solar system planets), any custom Elliptic ones for fantasy planets */
+	/** One of the predefined Orbit types (Solar system planets), any custom one for fantasy planets */
 	UPROPERTY(BlueprintReadWrite, EditAnywhere, Category= "Celestial Vault")
 	EOrbitType OrbitType = EOrbitType::Elliptic;
 	
-	/** The Body Radius in Kilometers */
-	UPROPERTY(BlueprintReadWrite, EditAnywhere, Category= "Celestial Vault", meta=(ForceUnits="Kilometers"))
-	double Radius = 1000.0;
-
-	/** The planetary body Material expects Bodies textures in a single row - This is the 0-based index of the texture to use */
-	UPROPERTY(BlueprintReadWrite, EditAnywhere, Category= "Celestial Vault")
-	int32 TextureColumnIndex = 0;
+	/** The Body J2000 Right Ascension in the Celestial Frame - In hours! */
+	UPROPERTY(BlueprintReadWrite, EditAnywhere, Category= "Celestial Vault", meta=(ForceUnits="Hours", DisplayName="Right Ascension J2000", MakeStructureDefaultValue="0.000000"))
+	double RAJ2000 = 0.0;
 
-	// TODO_Beta
-	// Add the Elliptic parameters for other Moon/Planets
-	// Ellipsoid - Minor/Major axes + offset to X axis
-	// Revolution period (year) + angle offset to time
-	// Phase if we want to fake it
-
-public: 
-	static FPlanetaryBodyInputData Earth; // Earth Preset
-	static FPlanetaryBodyInputData Moon; // Moon Preset
-};
-
-/**
- * Runtime structure to store the computed properties of a Planetary Body, for any Query. 
- */
-USTRUCT(BlueprintType)
-struct FPlanetaryBodyInfo
-{
-	GENERATED_BODY()
-public:
-	/** One of the predefined body types (Solar system planets), any custom one for fantasy planets */
-	UPROPERTY(BlueprintReadWrite, EditAnywhere, Category= "Celestial Vault")
-	EOrbitType OrbitType = EOrbitType::Elliptic;
-
-	/** The body Name */
-	UPROPERTY(BlueprintReadWrite, EditAnywhere, Category= "Celestial Vault")
-	FString Name = "";
+	/** The Body J2000 Declination in the Celestial Frame - In Degrees */
+	UPROPERTY(BlueprintReadWrite, EditAnywhere, Category= "Celestial Vault", meta=(ForceUnits="Degrees", DisplayName="Declination J2000", MakeStructureDefaultValue="0.000000"))
+	double DECJ2000 = 0.0;
 	
-	/** The Body Radius in Kilometers */
-	UPROPERTY(BlueprintReadWrite, EditAnywhere, Category= "Celestial Vault", meta=(ForceUnits="Kilometers"))
-	double Radius = 1000.0;
-
-	/** The Body Right Ascension in the Celestial Frame - In hours! */
+	/** The Body Right Ascension of date in the Celestial Frame - In hours! */
 	UPROPERTY(BlueprintReadWrite, EditAnywhere, Category= "Celestial Vault", meta=(ForceUnits="Hours"))
 	double RA = 0.0;
 
-	/** The Body Declination in the Celestial Frame - In Degrees */
+	/** The Body Declination of date in the Celestial Frame - In Degrees */
 	UPROPERTY(BlueprintReadWrite, EditAnywhere, Category= "Celestial Vault", meta=(ForceUnits="Degrees"))
 	double DEC  = 0.0;
 
@@ -97,113 +68,43 @@ public:
 	UPROPERTY(BlueprintReadWrite, EditAnywhere, Category= "Celestial Vault")
 	double DistanceInAU = 0.0;
 
-	/** The True apparent diameter of the body seen from Earth (Angular size, in degrees) */
-	UPROPERTY(BlueprintReadWrite, EditAnywhere, Category= "Celestial Vault",  meta = (ForceUnits="Degrees"))
-	double ApparentDiameterDegrees = 0.0;
-
-	/**
-	 * The Scaled apparent diameter of the body seen from Earth (Angular size, in degrees) 
-	 *   Takes the fake scaling factor into consideration, so useful for camera FOV tracking. 
-	 */
-	UPROPERTY(BlueprintReadWrite, EditAnywhere, Category= "Celestial Vault",  meta = (ForceUnits="Degrees"))
-	double ScaledApparentDiameterDegrees = 0.0;
-	
-	/** The Magnitude of the body seen from Earth */
-	UPROPERTY(BlueprintReadWrite, EditAnywhere, Category= "Celestial Vault")
-	double ApparentMagnitude = 0.0;
-
-	/** indication of the lunar age. 0 = New Moon, 0.25 = First quarter, 0.5 = Full Moon, 1 = Next New Moon */ 
-	UPROPERTY(BlueprintReadWrite, EditAnywhere, Category= "Celestial Vault")
-	double Age = 0.5;
-
-	/** indication of percentage of a full moon illumination */ 
-	UPROPERTY(BlueprintReadWrite, EditAnywhere, Category= "Celestial Vault")
-	double IlluminationPercentage = 1.0;
+	/** The Body Radius in Kilometers */
+	UPROPERTY(BlueprintReadWrite, EditAnywhere, Category= "Celestial Vault", meta=(ForceUnits="Kilometers"))
+	double Radius = 0.0;
 
-	/** Internal value of the ISM Instance corresponding to this planet */
+	/** The Magnitude of the Body seen from Earth */
 	UPROPERTY(BlueprintReadWrite, EditAnywhere, Category= "Celestial Vault")
-	int32 ISMInstanceIndex = 0; 
+	double Magnitude = 0.0;
 
 	/** Keep track of the Vault-relative transform, at ToD = 0 for animated bodies (Moon) */
 	UPROPERTY(BlueprintReadWrite, EditAnywhere, Category= "Celestial Vault")
-	FTransform UETransform;
+	FTransform UELocalTransform = FTransform::Identity;
 
 	/** Keep track of the location toward the Earth */ 
 	UPROPERTY(BlueprintReadWrite, EditAnywhere, Category= "Celestial Vault")
-	FVector DirectionTowardEarth = -FVector::ZeroVector;
+	FVector DirectionTowardEarth = FVector::ZeroVector;
 
 	// Can't use a UFUNCTION inside a struct --> Use the Celestial Maths Blueprint Function Library
 	FString ToString() const;
 
-	
-	void ComputeTranform(double UEDistance, double BodyScale);
-};
-
-/**
- * Runtime structure to store the computed properties of the Sun, for any Query. 
- */
-USTRUCT(BlueprintType)
-struct FSunInfo
-{
-	GENERATED_BODY()
-public:
-	/** The Sun Right Ascension in the Celestial Frame - In hours! */
-	UPROPERTY(BlueprintReadWrite, EditAnywhere, Category= "Celestial Vault", meta=(ForceUnits="Hours", DisplayName="Right Ascension", MakeStructureDefaultValue="0.000000"))
-	double RA = 0.0;
-
-	/** The Sun Declination in the Celestial Frame - In Degrees */
-	UPROPERTY(BlueprintReadWrite, EditAnywhere, Category= "Celestial Vault", meta=(ForceUnits="Degrees", DisplayName="Declination", MakeStructureDefaultValue="0.000000"))
-	double DEC = 0.0;
 
-	/** Keep track of the Vault-relative transform, at ToD = 0 for the Sun */
-	UPROPERTY(BlueprintReadWrite, EditAnywhere, Category= "Celestial Vault")
-	FTransform UETransform;
-
-	/** Keep track of the location toward the Earth */ 
-	UPROPERTY(BlueprintReadWrite, EditAnywhere, Category= "Celestial Vault")
-	FVector DirectionTowardEarth = -FVector::ZeroVector;;
-
-	// Can't use a UFUNCTION inside a struct --> Celestial Maths Blueprint Function Library
-	FString ToString() const;
-
-	// TODO_Beta
-	// Compute and store additional properties for the Sun
-	//   Elevation
-	//   CorrectedElevation - Sun Elevation, corrected for atmospheric diffraction
-	//   Azimuth
-	//   SunriseTime
-	//   SunsetTime
-	//   SolarNoon
+	
+	/** From the properties of this Celestial Body, compute the associated UE Transform, considering the following arguments
+	 *   - ExtraScale - By default, the base scale is computed by assuming the body is rendered by a 1meter billboard plane. This ExtraScale factor artificially makes the object larger.
+	 *   - DistanceOverride - if not set (or <0), we compute the UE Object location based on the actual Celestial Body Location. If set, we consider this distance instead. (Useful to place the body on a Plato-like sphere)
+	 */ 
+	virtual void ComputeUELocalTransform(double ExtraScale = 1.0, double UEDistanceOverride = -1.0 );
 };
 
-/**
- * Runtime structure to store the computed properties of a Star, for any Query. 
- */
+/** Additional Properties specific to Stellar Bodies (Stars) */
 USTRUCT(BlueprintType)
-struct FStarInfo
+struct FStellarBody : public FCelestialBody
 {
 	GENERATED_BODY()
 
-public:
-	/** The Star Right Ascension in the Celestial Frame - In hours! */
-	UPROPERTY(BlueprintReadWrite, EditAnywhere, Category= "Celestial Vault", meta=(ForceUnits="Hours", DisplayName="Right Ascension", MakeStructureDefaultValue="0.000000"))
-	double RA = 0.0;
-
-	/** The Star Declination in the Celestial Frame - In Degrees */
-	UPROPERTY(BlueprintReadWrite, EditAnywhere, Category= "Celestial Vault", meta=(ForceUnits="Degrees", DisplayName="Declination", MakeStructureDefaultValue="0.000000"))
-	double DEC = 0.0;
-
-	/** Earth to Star distance (in Parsecs)  */
-	UPROPERTY(BlueprintReadWrite, EditAnywhere, Category= "Celestial Vault", meta=(DisplayName="Distance (in parsecs)", MakeStructureDefaultValue="100.000000"))
-	double DistanceInPC = 100.0;
-	
-	/** Star Name */
-	UPROPERTY(BlueprintReadWrite, EditAnywhere, Category= "Celestial Vault", meta=(DisplayName="Name"))
-	FString Name = "";
-
-	/** Star Magnitude */
-	UPROPERTY(BlueprintReadWrite, EditAnywhere, Category= "Celestial Vault", meta=(DisplayName="Magnitude", MakeStructureDefaultValue="1.000000"))
-	double Magnitude = 1.0;
+	/** Star Color Index, also named B-V */
+	UPROPERTY(BlueprintReadWrite, EditAnywhere, Category= "Celestial Vault", meta=(DisplayName="ColorIndex", MakeStructureDefaultValue="0.000000"))
+	double ColorIndex = 0.0;
 
 	/** Star RGB Color - Can be computed from the B-V value if the star is from an official Catalog */
 	UPROPERTY(BlueprintReadWrite, EditAnywhere, Category= "Celestial Vault", meta=(DisplayName="Color"))
@@ -220,78 +121,44 @@ public:
 	/** Star YaleBrightStar ID if present in the Yale Bright Star Catalog */
 	UPROPERTY(BlueprintReadWrite, EditAnywhere, Category= "Celestial Vault", meta=(DisplayName="Yale Bright Star Catalog ID", MakeStructureDefaultValue="0"))
 	int32 YaleBrightStarID = 0;
-	
-	/** Star Color Index, also named B-V */
-	UPROPERTY(BlueprintReadWrite, EditAnywhere, Category= "Celestial Vault", meta=(DisplayName="ColorIndex", MakeStructureDefaultValue="0.000000"))
-	double ColorIndex = 0.0;
 
-	/** Internal value of the ISM Instance corresponding to this Star */
-	UPROPERTY(BlueprintReadWrite, EditAnywhere, Category= "Celestial Vault",  meta = (HideInDataTable))
-	int32 ISMInstanceIndex = 0;
-
-	// Can't use a UFUNCTION inside a struct --> Celestial Maths Blueprint Function Library
+	// Can't use a UFUNCTION inside a struct --> Use the Celestial Maths Blueprint Function Library
 	FString ToString() const;
 };
 
-
-/**
- * TableRow Base type to describe the Creation parameters of a Basic Star
- * Will only be used at creation time, from a proper Data Table
- * This Struct contains the minimal needed data for a Fictional Star
- */
+/** Additional Properties specific to Planetary Bodies (Planets/Moons) */
 USTRUCT(BlueprintType)
-struct FStarInputData : public FTableRowBase 
+struct FPlanetaryBody : public FCelestialBody
 {
 	GENERATED_BODY()
-public:
-	/** The Star Right Ascension in the Celestial Frame - In hours! */
-	UPROPERTY(BlueprintReadWrite, EditAnywhere, Category= "Celestial Vault", meta=(ForceUnits="Hours", DisplayName="Right Ascension", MakeStructureDefaultValue="0.000000"))
-	double RA = 0.0;
 
-	/** The Star Declination in the Celestial Frame - In Degrees */
-	UPROPERTY(BlueprintReadWrite, EditAnywhere, Category= "Celestial Vault", meta=(ForceUnits="Degrees", DisplayName="Declination", MakeStructureDefaultValue="0.000000"))
-	double DEC = 0.0;
+	/** VSOP87 Identifier of the Planetary Body */
+	UPROPERTY(BlueprintReadWrite, EditAnywhere, Category= "Celestial Vault")
+	EVSOP87BodyType VSOP87BodyType = EVSOP87BodyType::Earth;
 
-	/** Earth to Star distance (in Parsecs)  */
-	UPROPERTY(BlueprintReadWrite, EditAnywhere, Category= "Celestial Vault", meta=(DisplayName="Distance (in parsecs)", MakeStructureDefaultValue="100.000000"))
-	double DistanceInPC = 100.0;
+	/** The True apparent diameter of the Body seen from Earth (Angular size, in degrees) */
+	UPROPERTY(BlueprintReadWrite, EditAnywhere, Category= "Celestial Vault",  meta = (ForceUnits="Degrees"))
+	double ApparentDiameterDegrees = 0.0;
+
+	/**
+	 * The Scaled apparent diameter of the Body seen from Earth (Angular size, in degrees) 
+	 *   Takes the fake scaling factor into consideration, so useful for camera FOV tracking. 
+	 */
+	UPROPERTY(BlueprintReadWrite, EditAnywhere, Category= "Celestial Vault",  meta = (ForceUnits="Degrees"))
+	double ScaledApparentDiameterDegrees = 0.0;
 	
-	/** Star Name */
-	UPROPERTY(BlueprintReadWrite, EditAnywhere, Category= "Celestial Vault", meta=(DisplayName="Name"))
-	FString Name = "";
+	/** indication of the body disc age (phase). 0 = New, 0.25 = First quarter, 0.5 = Full, 1 = Next New */ 
+	UPROPERTY(BlueprintReadWrite, EditAnywhere, Category= "Celestial Vault")
+	double Age = 0.5;
 
-	/** Star Magnitude */
-	UPROPERTY(BlueprintReadWrite, EditAnywhere, Category= "Celestial Vault", meta=(DisplayName="Magnitude", MakeStructureDefaultValue="1.000000"))
-	double Magnitude = 1.0;
+	/** indication of percentage of the illumination when the Body is fully visible*/ 
+	UPROPERTY(BlueprintReadWrite, EditAnywhere, Category= "Celestial Vault")
+	double IlluminationPercentage = 1.0;
 
-	/** Star RGB Color - Useless if computed from the B-V */ 
-	UPROPERTY(BlueprintReadWrite, EditAnywhere, Category= "Celestial Vault", meta=(DisplayName="Color"))
-	FLinearColor Color = FLinearColor(1.0f,1.0f,1.0f,1.0f);
+	FString ToString() const;
+
+	virtual void ComputeUELocalTransform(double ExtraScale = 1.0, double UEDistanceOverride = -1.0 ) override;
 };
 
-/**
- * TableRow Base type to describe the Creation parameters of a Catalog-based Star
- * Will only be used at creation time, from a proper Data Table
- * This Struct extends the FStarInputData class with additional Catalog Properties
- */
-USTRUCT(BlueprintType)
-struct FCelestialStarInputData : public FStarInputData 
-{
-	GENERATED_BODY()
-public:
-	/** Star Hipparcos ID */
-	UPROPERTY(BlueprintReadWrite, EditAnywhere, Category= "Celestial Vault", meta=(DisplayName="Hipparcos Catalog ID", MakeStructureDefaultValue="0"))
-	int32 HipparcosID = -1;
 
-	/** Star Henry Draper ID */
-	UPROPERTY(BlueprintReadWrite, EditAnywhere, Category= "Celestial Vault", meta=(DisplayName="Henry Draper Catalog ID", MakeStructureDefaultValue="0"))
-	int32 HenryDraperID = -1;
 
-	/** Star YaleBrightStar ID */
-	UPROPERTY(BlueprintReadWrite, EditAnywhere, Category= "Celestial Vault", meta=(DisplayName="Yale Bright Star Catalog ID", MakeStructureDefaultValue="0"))
-	int32 YaleBrightStarID = -1;
-	
-	/** Star Color Index, also named B-V */
-	UPROPERTY(BlueprintReadWrite, EditAnywhere, Category= "Celestial Vault", meta=(DisplayName="ColorIndex", MakeStructureDefaultValue="0.000000"))
-	double ColorIndex = -1.0;
-};
\ No newline at end of file
diff --git a/Engine/Plugins/Experimental/CelestialVault/Source/CelestialVault/Public/CelestialMaths.h b/Engine/Plugins/Experimental/CelestialVault/Source/CelestialVault/Public/CelestialMaths.h
--- a/Engine/Plugins/Experimental/CelestialVault/Source/CelestialVault/Public/CelestialMaths.h
+++ b/Engine/Plugins/Experimental/CelestialVault/Source/CelestialVault/Public/CelestialMaths.h
@@ -31,32 +31,44 @@ class CELESTIALVAULT_API UCelestialMaths : public UBlueprintFunctionLibrary
 public:
 
 #pragma region Colors
-    
+
     /** Returns the RGB normalized components [0..1] from the Color Index (B-V) Value **/
     UFUNCTION(BlueprintCallable, Category="Celestial Vault|Color")
     static FLinearColor BVtoLinearColor(float BV);
 
+	/** return the illumination factor (0..1) of a Body, considering his normalized age and the crescent effects */
+	UFUNCTION(BlueprintCallable, Category="Celestial Vault|Color")
+	static double GetIlluminationPercentage(double NormalizedAge);
+
 #pragma endregion
 
-#pragma region Planetary Bodies
+#pragma region VSOP87
 
-    /**
-     * Returns the location of a specific Planetary body, in the FK5 J2000 Coordinate System
+	    /**
+     * Returns the location of a specific Solar System body, in the VSOP Coordinate system
      * 
-     *   The FK5 is an equatorial coordinate system (coordinate system linked to the Earth) based on its J2000 position.
-     *   As any equatorial frame, the FK5-based follows the long-term Earth motion (precession).
-     *       
+     *   The VSOP frame is a "heliocentric ecliptic frame of the J2000 equinox"
+     *     Origin = Sun (Heliocentric)
+     *     Coordinates Axes = Mean ecliptic & equinox of J2000.0
+	 *        XY-plane: the mean ecliptic plane at epoch J2000.0
+     *        X-axis: pointing toward mean equinox at J2000.0
+     *        Z-axis: perpendicular to the ecliptic
+     *        
      *   The returned location is expressed in Astronomical Units (AU)
      *   The relativistic effects are ignored (See GetBodyLocation_FK5J2000_AU_Relativistic )
      */
-    UFUNCTION(BlueprintCallable, Category="Celestial Vault|Planetary bodies")
-    static FVector GetBodyLocation_FK5J2000_AU(FPlanetaryBodyInputData PlanetaryBody, double JulianDate);
+    UFUNCTION(BlueprintCallable, Category="Celestial Vault|VSOP87")
+    static void GetSolarSystemBodyLocationVelocity_VSOP87_AU(EVSOP87BodyType VSOP87BodyType, double JulianDate, FVector& LocationAU, FVector& Velocity);
 
     /**
-     * Returns the location of a specific Planetary body, in the FK5 J2000 Coordinate System
+     * Returns the location of a specific  Solar System Body, in the VSOP Coordinate system
      * 
-     *   The FK5 is an equatorial coordinate system (coordinate system linked to the Earth) based on its J2000 position.
-     *   As any equatorial frame, the FK5-based follows the long-term Earth motion (precession).
+	 *   The VSOP frame is a "heliocentric ecliptic frame of the J2000 equinox"
+	 *     Origin = Sun (Heliocentric)
+	 *     Coordinates Axes = Mean ecliptic & equinox of J2000.0
+	 *        XY-plane: the mean ecliptic plane at epoch J2000.0
+	 *        X-axis: pointing toward mean equinox at J2000.0
+	 *        Z-axis: perpendicular to the ecliptic
      *       
      *   The returned location is expressed in Astronomical Units (AU)
      *       
@@ -64,32 +76,25 @@ public:
      * The relativistic effects are considered (time taken for the light to reach the Observer body location)
      * This function therefore returns the location of the Planetary Body as if it was seen from the Observer Body Location
     */
-    UFUNCTION(BlueprintCallable, Category="Celestial Vault|Planetary bodies")
-    static FVector GetBodyLocation_FK5J2000_AU_Relativistic(FVector ObserverBodyFK5J2000LocationAU, FPlanetaryBodyInputData PlanetaryBody, double JulianDate);
+    UFUNCTION(BlueprintCallable, Category="Celestial Vault|VSOP87")
+    static FVector GetSolarSystemBodyLocation_VSOP87_Relativistic(FVector ReferenceBodyLocation_Heliocentric_AU, EVSOP87BodyType VSOP87BodyType, double JulianDate);
+
+	
+#pragma endregion
+
+#pragma region Celestial Bodies
 
     /** Return the location of a Planetary Body relative to the Earth, expressed in Celestial Coordinates (RA, DEC, Distance)
      * It requires the Observer location on Earth for more precise computations
-     * This function also returns the distance to the Sun, as it can help for Magnitude computations
+     * This function also returns the distance between bodies, as it can help for Magnitude computations
      * **/
     UFUNCTION(BlueprintCallable, Category="Celestial Vault|Planetary bodies")
-    static void GetBodyCelestialCoordinatesAU(double JulianDate, FPlanetaryBodyInputData PlanetaryBody, double ObserverLatitude, double ObserverLongitude, double& RAHours, double& DECDegrees,
-        double& DistanceBodyToEarthAU, double& DistanceBodyToSunAU, double& DistanceEarthToSunAU);
+    static void GetBodyCelestialCoordinatesAU(double JulianDate, EVSOP87BodyType VSOP87BodyType, double ObserverLatitude, double ObserverLongitude,
+    	double& RAJ2000Hours, double& DECJ2000Degrees, double& RAHours, double& DECDegrees, double& DistanceBodyToEarthAU, double& DistanceBodyToSunAU, double& DistanceEarthToSunAU);
 
     /** Return the Magnitude of a Planetary Body as seen from the Earth **/
     UFUNCTION(BlueprintCallable, Category="Celestial Vault|Planetary bodies")
-    static double GetPlanetaryBodyMagnitude(FPlanetaryBodyInputData PlanetaryBody, double DistanceToSunAU, double DistanceToEarthAU, double DistanceEarthToSunAU, double& PhaseAngle);
-
-    /** Returns the Moon Phase for a specific Date
-     * 
-     *  This is an approximate computation using a number of lunar cycles with a synodic month equals to 29.53059 days
-     *  Not very precise over more than 1 centuries before of after 2025
-     */
-    UFUNCTION(BlueprintCallable, Category="Celestial Vault|Planetary bodies")
-    static double GetMoonNormalizedAgeSimple(double JulianDate);
-
-	/** return the illumination factor (0..1) of a Body, considering his normalized age and the crescent effects */
-	UFUNCTION(BlueprintCallable, Category="Celestial Vault|Planetary bodies")
-	static double GetIlluminationPercentage(double NormalizedAge);
+    static double GetPlanetaryBodyMagnitude(EVSOP87BodyType VSOP87BodyType, double DistanceToSunAU, double DistanceToEarthAU, double DistanceEarthToSunAU, double& PhaseAngle);
 
 #pragma endregion
 
@@ -99,7 +104,7 @@ public:
     UFUNCTION(BlueprintCallable, Category="Celestial Vault|Time")
     static FDateTime LocalTimeToUTCTime(FDateTime LocalTime, double TimeZoneOffset, bool IsDst);
 
-    /** Return the UTC for a specific Time, using the TimeZone and Daylight Saving Information **/
+    /** Return the Local Time for a specific UTC Time, using the TimeZone and Daylight Saving Information **/
     UFUNCTION(BlueprintCallable, Category="Celestial Vault|Time")
     static FDateTime UTCTimeToLocalTime(FDateTime UTCTime, double TimeZoneOffset, bool IsDst);
     
@@ -167,12 +172,24 @@ public:
     UFUNCTION(BlueprintCallable, Category="Celestial Vault|Time")
     static double JulianDateToInternationalAtomicTime(double JulianDate);
 
+	/** Converts Seconds to Decimal Days */
     UFUNCTION(BlueprintCallable, Category="Celestial Vault|Time")
     static inline double SecondsToDay(double Seconds) { return Seconds / 86400.0; }
 
+	/** Converts decimal Days to Seconds */
     UFUNCTION(BlueprintCallable, Category="Celestial Vault|Time")
     static inline double DaysToSeconds(double Days) { return Days * 86400.0; }
 
+    /** Converts a Julian Date to a proper VSOP87Time, suitable for the VSOP87 computations
+     * 
+	 * In the VSOP87 equations, Time is not Julian centuries!
+	 *  > Given a Julian date JD expressed in dynamical time T = (TAI+32.184s), the VSOP equations exprect a Time = (T-2451545)/365250
+	 *	> The denominator is 365250, not 36525.
+	 *	> That’s because VSOP87 expresses time in units of 10,000 Julian years (1 Julian millennium = 100 centuries = 36,525 × 100 days).
+     */
+    UFUNCTION(BlueprintCallable, Category="Celestial Vault|Time")
+    static double JulianDateToVSOP87Time(double JulianDate);
+	
 #pragma endregion
 
 #pragma region Angles
@@ -247,22 +264,58 @@ public:
 
     /** Approximation of the IAU2000A/B nutation model used in the Equation Of The Equinoxes, accurate enough for VSOP87 computations  */
     UFUNCTION(BlueprintCallable, Category="Celestial Vault|Earth")
-    static void Nutation2000BTruncated(double JulianDate, double& DeltaPsi, double& DeltaEpsilon);
+    static void Nutation2000BTruncated(double JulianDate, double& DeltaPsiDegrees, double& DeltaEpsilonDegrees);
 
     /** Return the Transformation to apply to a WGS84 Ellipsoid model so that its location in Lat,long,Altitude is tangent to the Origin
-     * It's used to locate the Rotating Celestial Vault for a specific UE Origin
+     *  The Transform is expressed in the ECEF Frame (Location units are meters, right-handed ECEF Frame)
      */
     UFUNCTION(BlueprintCallable, Category="Celestial Vault|Earth")
-    static FTransform GetPlanetCenterTransform(double Latitude, double Longitude, double Altitude);
-    
+    static FTransform GetEarthCenterTransformECEF(double Latitude, double Longitude, double Altitude);
+
+	/** Return the Transformation to apply to a WGS84 Ellipsoid model so that its location in Lat,long,Altitude is tangent to the Origin
+	 * It's used to locate the Rotating Celestial Vault for a specific UE Origin
+	 * The Transform is expressed in the UE Frame (Location units are UE Units, left-handed Unreal Frame)
+	 */
+	UFUNCTION(BlueprintCallable, Category="Celestial Vault|Earth")
+	static FTransform GetEarthCenterTransformUEFrame(double Latitude, double Longitude, double Altitude);
+
+	/** Return the Precession Matrix 
+	 *    The Precession matrix transforms coordinates from one equatorial reference epoch to another by accounting for Earth’s long-term, secular drift of its rotation axis (precession of the equinox).
+	 *    It's used to compute RA, DEC between the J2000 Epoch and a specific date. (Earth’s rotation axis slowly traces a ~26,000-year cone because of gravitational torques from the Moon and Sun.)
+	 */
+	UFUNCTION(BlueprintCallable, Category="Celestial Vault|Earth")
+	static FMatrix GetPrecessionMatrix(double JulianDate);
+
+	/** Return the Nutation Matrix 
+	 *   The Nutation matrix accounts for Nutation, the small, periodic wobble of Earth’s rotation axis – when transforming between “mean” and “true” celestial coordinate frames.
+	 */
+	UFUNCTION(BlueprintCallable, Category="Celestial Vault|Earth")
+	static FMatrix GetNutationMatrix(double JulianDate);
+
+    /* Compute the Aberration of an object located far from the Earth.
+	 *   In astronomy, Aberration is an apparent shift in the observed position of a celestial object caused by the finite speed of light combined with the motion of the observer.
+     *   It is not an optical distortion — it’s a relativistic effect due to the fact that light takes time to reach you while you are moving.
+     */
+    UFUNCTION(BlueprintCallable, Category="Celestial Vault|Earth")
+	static FVector ComputeAberration(const FVector& TopocentricTargetLocationAU, const FVector& EarthVelocity);
+
+	
 #pragma endregion
     
-#pragma region Sun
+#pragma region Sun & Moon
 
-    /** Compute all Sun Properties for a speficic JulianDate */ 
+    /** Compute all Sun Properties for a specific JulianDate */ 
     UFUNCTION(BlueprintCallable, Category="Celestial Vault|Sun")
-    static FSunInfo GetSunInformation(double JulianDate, double ObserverLatitude, double ObserverLongitude);
-        
+    static FStellarBody GetSunInformation(double JulianDate, double ObserverLatitude, double ObserverLongitude);
+
+	/** Returns the Moon Phase for a specific Date
+	 * 
+	 *  This is an approximate computation using a number of lunar cycles with a synodic month equals to 29.53059 days
+	 *  Not very precise over more than 1 centuries before of after 2025
+	 */
+	UFUNCTION(BlueprintCallable, Category="Celestial Vault|Planetary bodies")
+	static double GetMoonNormalizedAgeSimple(double JulianDate);
+	
 #pragma endregion
 
 #pragma region Utilities
@@ -270,7 +323,6 @@ public:
     /** Returns the Speed of Light 299 792 458 (m/s) **/
     UFUNCTION(BlueprintCallable, Category="Celestial Vault|Utilities")
     static double GetSpeedOfLight() { return UCelestialMaths::SpeedOfLightMetersPerSeconds; }
-
     
     /** Convert Astronomical Unit (UA) to meters
      *
@@ -286,6 +338,20 @@ public:
     UFUNCTION(BlueprintCallable, Category="Celestial Vault|Utilities")
     static double MetersToAstronomicalUnits(double Meters) {return Meters / AstronomicalUnitsMeters; };
 
+	/** Convert Parsecs to Astronomical Unit (UA)
+	 *
+	 * 1 Parsec = 206 264.806 AU
+	 */
+	UFUNCTION(BlueprintCallable, Category="Celestial Vault|Utilities")
+	static double ParsecsToAstronomicalUnits(double Parsecs) {return Parsecs * ParsecAstronomicalUnits; };
+
+	/** Convert Astronomical Unit (UA) to Parsecs
+	 *
+	 * 1 Parsec = 206 264.806 AU
+	 */
+	UFUNCTION(BlueprintCallable, Category="Celestial Vault|Utilities")
+	static double AstronomicalUnitsToParsec(double AU) {return AU / ParsecAstronomicalUnits; };
+
     /** Convert Cartesian Coordinates to Polar Coordinates, using a Righ-Handed Frame */
     UFUNCTION(BlueprintCallable, Category="Celestial Vault|Utilities")
     static void XYZToRADEC_RH(FVector XYZ, double& RADegrees, double& DECDegrees, double& Radius);
@@ -298,17 +364,17 @@ public:
     UFUNCTION(BlueprintCallable, Category="Celestial Vault|Utilities")
     static FString GetPreciseVectorString(FVector Vector, int32 MinimumFractionalDigits = 10);
 
-    /** StarInfo String Builder */
-    UFUNCTION(BlueprintPure, meta = (DisplayName = "To String (StarInfo)", CompactNodeTitle = "->", BlueprintAutocast), Category = "Celestial|Utilities")
-    static FString Conv_StarInfoToString(const FStarInfo& StarInfo);
+    /** Celestial Body String Builder */
+    UFUNCTION(BlueprintPure, meta = (DisplayName = "To String (CelestialBody)", CompactNodeTitle = "->", BlueprintAutocast), Category = "Celestial|Utilities")
+    static FString Conv_CelestialBodyToString(const FCelestialBody& CelestialBody); 
 
-    /** PlanetaryBodyInfo String Builder */
-    UFUNCTION(BlueprintPure, meta = (DisplayName = "To String (PlanetaryBodyInfo)", CompactNodeTitle = "->", BlueprintAutocast), Category = "Celestial|Utilities")
-    static FString Conv_PlanetaryBodyInfoToString(const FPlanetaryBodyInfo& PlanetaryBodyInfo);
+    /** Planetary Body String Builder */
+    UFUNCTION(BlueprintPure, meta = (DisplayName = "To String (PlanetaryBody)", CompactNodeTitle = "->", BlueprintAutocast), Category = "Celestial|Utilities")
+    static FString Conv_PlanetaryBodyToString(const FPlanetaryBody& PlanetaryBody);
 
-    /** SunInfo ToString String Builder */
-    UFUNCTION(BlueprintPure, meta = (DisplayName = "To String (SunInfo)", CompactNodeTitle = "->", BlueprintAutocast), Category = "Celestial|Utilities")
-    static FString Conv_SunInfoToString(const FSunInfo& SunInfo);
+    /** Stellar Body ToString String Builder */
+    UFUNCTION(BlueprintPure, meta = (DisplayName = "To String (StellarBody)", CompactNodeTitle = "->", BlueprintAutocast), Category = "Celestial|Utilities")
+    static FString Conv_StellarBodyToString(const FStellarBody& StellarBody); 
 	
     /** Right Ascension String Builder */ 
     UFUNCTION(BlueprintPure, meta = (DisplayName = "To String (Right Ascension in hours)"), Category = "Celestial|Utilities")
@@ -325,11 +391,19 @@ public:
     /** DMS String Builder */
     UFUNCTION(BlueprintPure, meta = (DisplayName = "To String (Signed Angle in Degrees, Minutes, Seconds )"), Category = "Celestial|Utilities")
     static FString Conv_DMSToString(bool Sign, int32 Degrees, int32 Minutes, double Seconds);
+
+	/** HMS String Builder */
+	UFUNCTION(BlueprintPure, meta = (DisplayName = "To HMS String (Angle in Degrees)"), Category = "Celestial|Utilities")
+	static FString Conv_DegreesToHMSString(double Degrees);
+	
+	/** DMS String Builder */
+	UFUNCTION(BlueprintPure, meta = (DisplayName = "To DMS String (Angle in Degrees)"), Category = "Celestial|Utilities")
+	static FString Conv_DegreesToDMSString(double Degrees);
     
 #pragma endregion
-public:
-    static double SynodicMonthAverage;
-    
+
+#pragma region Private Functions and variables
+	
 private:
 
     /** Because the Dates are expressed around Jan 1, 4713 BC, we need a special Floor function to have the right years continuity
@@ -338,11 +412,24 @@ private:
      */
     static int32 FloorForJulianDate(double JulianDate);
 
+	/* Return a Rotation Matrix around X, for Precession/Nutation computations */ 
+	static FMatrix GetRotationXMatrix(double AngleRadians);
+
+	/* Return a Rotation Matrix around Y, for Precession/Nutation computations */
+	static FMatrix GetRotationYMatrix(double AngleRadians);
+
+	/* Return a Rotation Matrix around Z, for Precession/Nutation computations */
+	static FMatrix GetRotationZMatrix(double AngleRadians);
+	
+
 private:
     // Static const Members
-    static FMatrix VSOPToJ2000;
-    static double SpeedOfLightMetersPerSeconds;
-    static double AstronomicalUnitsMeters;
-    static double NewMoonReferenceJulianDate;
-    
+	static const double SynodicMonthAverage ;
+    static const FMatrix VSOPEclipticToFK5Equatorial;
+    static const double SpeedOfLightMetersPerSeconds;
+    static const double AstronomicalUnitsMeters;
+    static const double ParsecAstronomicalUnits;
+	static const double NewMoonReferenceJulianDate;
+
+	#pragma endregion
 };
diff --git a/Engine/Plugins/Experimental/CelestialVault/Source/CelestialVault/Public/CelestialVaultDaySequenceActor.h b/Engine/Plugins/Experimental/CelestialVault/Source/CelestialVault/Public/CelestialVaultDaySequenceActor.h
--- a/Engine/Plugins/Experimental/CelestialVault/Source/CelestialVault/Public/CelestialVaultDaySequenceActor.h
+++ b/Engine/Plugins/Experimental/CelestialVault/Source/CelestialVault/Public/CelestialVaultDaySequenceActor.h
@@ -4,9 +4,18 @@
 
 #include "CoreMinimal.h"
 #include "CelestialDataTypes.h"
+#include "CelestialInputDataTypes.h"
 #include "DaySequenceActor.h"
 #include "CelestialVaultDaySequenceActor.generated.h"
 
+#ifndef CELESTIAL_VAULT_ENABLE_DRAW_DEBUG
+	#define CELESTIAL_VAULT_ENABLE_DRAW_DEBUG !UE_BUILD_SHIPPING
+#endif
+
+#if CELESTIAL_VAULT_ENABLE_DRAW_DEBUG
+	// class ;
+#endif
+
 // Celestial Classes
 class UStarsDataTable;
 
@@ -159,7 +168,7 @@ public: // Properties
 	
 	/** Array of the created Stars information - Only populated if KeepStarsInfo is true */ 
 	UPROPERTY(BlueprintReadOnly, Category="Celestial Vault|Stars")
-	TArray<FStarInfo> StarsInfo;
+	TArray<FStellarBody> Stars;
 	
 #pragma endregion
 	
@@ -179,7 +188,7 @@ public: // Properties
 	
 	/** Array of the created planetary bodies, with all their computed information - Only populated if KeepPlanetsInfos is true */ 
 	UPROPERTY(BlueprintReadOnly, Category="Celestial Vault|Planets")
-	TArray<FPlanetaryBodyInfo> PlanetsInfos;
+	TArray<FPlanetaryBody> Planets;
 
 #pragma endregion
 
@@ -191,7 +200,7 @@ public: // Properties
 
 	/** Celestial Info for the Moon, at the beginning of the Day */  
 	UPROPERTY(BlueprintReadOnly, Category="Celestial Vault|Moons")
-	FPlanetaryBodyInfo MoonBodyInfo;
+	FPlanetaryBody Moon;
 
 	/** If true, the moon Age (Phase) and location can be overriden */  
 	UPROPERTY(EditAnywhere, Category="Celestial Vault|Moons")
@@ -219,7 +228,7 @@ public: // Properties
 
 	/** Celestial Info for the Sun, at the beginning of the Day */  
 	UPROPERTY(BlueprintReadOnly, Category="Celestial Vault|Sun")
-	FSunInfo SunInfo; 
+	FStellarBody Sun; 
 
 	/** Base Sun Intensity
 	 * Typically 120000 Lux
@@ -232,8 +241,6 @@ public: // Properties
 	 */
 	UPROPERTY(EditAnywhere, BlueprintReadWrite, AdvancedDisplay, Category="Celestial Vault", meta=( ForceUnits="Lux", ClampMin = "0", ClampMax = "200000"))
 	float MoonLightIntensity = 0.1f;
-
-	
 #pragma endregion 
 
 protected: // Functions
@@ -241,6 +248,10 @@ protected: // Functions
 	/** BeginPlay and OnConstruction overrides auto-register this actor with the DaySequenceSubsystem. */
 	virtual void BeginPlay() override;
 	virtual void OnConstruction(const FTransform& Transform) override;
+
+public:
+	virtual void PostRegisterAllComponents() override;
+	virtual void UnregisterAllComponents(bool bForReregister = false) override;
 	
 public: // Functions
 
@@ -250,11 +261,11 @@ public: // Functions
 
 	/** Returns the Celestial Info for the Sun, at a specific Julian Date */  
 	UFUNCTION(BlueprintCallable, Category = "Celestial Vault|Sun")
-	FSunInfo GetSunInfo(double JulianDate);
+	FStellarBody GetSun(double JulianDate) const; 
 	
 	/** Returns the Celestial Info for the Moon, at a specific Julian Date */  
 	UFUNCTION(BlueprintCallable, Category = "Celestial Vault|Queries")
-	FPlanetaryBodyInfo GetMoonInfo(double JulianDate);
+	FPlanetaryBody GetMoon(double JulianDate) const;
 
 	/** Manually set the Moon Age (Phase) */  
 	UFUNCTION(BlueprintCallable, Category = "Celestial Vault|Moons")
@@ -264,15 +275,15 @@ public: // Functions
 	
 	/** Return the Celestial Information of the Star closest to a specific direction, within an angle threshold */
 	UFUNCTION(BlueprintCallable, Category = "Celestial Vault|Queries")
-	bool GetClosestStarInfo(FVector ObserverLocation, FVector LookupDirection, double ThresholdAngleDegree, FStarInfo& FoundStarInfo, FTransform& StarTransform);
+	bool GetClosestStar(FVector ObserverLocation, FVector LookupDirection, double ThresholdAngleDegree, FStellarBody& FoundStarInfo, FTransform& StarTransform);
 	
 	/** Return the Celestial Information of the Planetary Body (moon, planet) closest to a specific direction, within an angle threshold */
 	UFUNCTION(BlueprintCallable, Category = "Celestial Vault|Queries")
-	bool GetClosestPlanetaryBody(FVector StartPosition, FVector LookupDirection, double ThresholdAngleDegree, FPlanetaryBodyInfo& FoundPlanetaryBodyInfo, FTransform& BodyTransform);
+	bool GetClosestPlanetaryBody(FVector StartPosition, FVector LookupDirection, double ThresholdAngleDegree, FPlanetaryBody& FoundPlanetaryBodyInfo, FTransform& BodyTransform);
 
 	/** Return the Celestial Information of a specific Planetary Body (by its orbit type) */
 	UFUNCTION(BlueprintCallable, Category = "Celestial Vault|Queries")
-	bool GetPlanetaryBodyByOrbitType(EOrbitType OrbitType, FPlanetaryBodyInfo& FoundPlanetaryBodyInfo, FTransform& BodyTransform);
+	bool GetPlanetaryBodyByVSOP87Type(EVSOP87BodyType VSOP87Type, FPlanetaryBody& FoundPlanetaryBodyInfo, FTransform& BodyTransform);
 
 	/** Returns the Celestial vault rotation angle for the current Date at t=0 midnight */
 	double GetDayCelestialVaultAngle() const;
@@ -285,10 +296,15 @@ public: // Functions
 private:
 	void InitStars();
 	void InitPlanetaryBodies();
-	FPlanetaryBodyInfo GetPlanetaryBodyInfo(const FPlanetaryBodyInputData& InputPlanetaryBody, double JulianDay, double UEDistance, double BodyScale) const;
+	FPlanetaryBody GetPlanetaryBodyInfo(const FPlanetaryBodyInputData& InputPlanetaryBody, double JulianDay) const;
 
 	UFUNCTION(BlueprintCallable, CallInEditor, Category= "Celestial Vault")
 	void RebuildAll();
 
-	FDelegateHandle DrawDebugDelegateHandle;
+#if CELESTIAL_VAULT_ENABLE_DRAW_DEBUG
+	const FString CelestialVaultSF = TEXT("CelestialVault");
+	FDelegateHandle DrawDebugDelegateHandle ;
+	void DebugDrawCallback(UCanvas* Canvas, APlayerController* PC) const;
+#endif
+	
 };
diff --git a/Engine/Plugins/Experimental/CelestialVault/Source/CelestialVaultEditor/CelestialVaultEditor.Build.cs b/Engine/Plugins/Experimental/CelestialVault/Source/CelestialVaultEditor/CelestialVaultEditor.Build.cs
--- a/Engine/Plugins/Experimental/CelestialVault/Source/CelestialVaultEditor/CelestialVaultEditor.Build.cs
+++ b/Engine/Plugins/Experimental/CelestialVault/Source/CelestialVaultEditor/CelestialVaultEditor.Build.cs
@@ -11,7 +11,7 @@ public class CelestialVaultEditor : ModuleRules
         PublicDependencyModuleNames.AddRange(
             new string[]
             {
-                "Core",
+                "Core"
             }
         );
 
@@ -21,7 +21,8 @@ public class CelestialVaultEditor : ModuleRules
                 "CoreUObject",
                 "Engine",
                 "Slate",
-                "SlateCore"
+                "SlateCore",
+                "UnrealEd", "CelestialVault"
             }
         );
     }
