- Added a system to automatically compute the Daylight saving times from their official rules, and adapt the vault state automatically.
- Added a Global North Offset value to compensate for non-North-aligned environments in Architecture.
- Adjust the SkyAtmosphere radius to the actual planet radius at the origin.
- Improved the property categories of the Celestial Vault Actor

#virtualized
#Synced-CL 49450045

diff --git a/Engine/Plugins/Experimental/CelestialVault/Source/CelestialVault/Private/CelestialMaths.cpp b/Engine/Plugins/Experimental/CelestialVault/Source/CelestialVault/Private/CelestialMaths.cpp
--- a/Engine/Plugins/Experimental/CelestialVault/Source/CelestialVault/Private/CelestialMaths.cpp
+++ b/Engine/Plugins/Experimental/CelestialVault/Source/CelestialVault/Private/CelestialMaths.cpp
@@ -447,10 +447,14 @@ FDateTime UCelestialMaths::JulianDateToUTCDateTime(double JulianDate)
 
 double UCelestialMaths::DateTimeToGreenwichMeanSiderealTime(FDateTime UTCDateTime)
 {
-	double JulianDate = UTCDateTimeToJulianDate(UTCDateTime);
+    return ModPositive(DateTimeToGreenwichMeanSiderealTimeUnwrapped(UTCDateTime), 360.0);
+}
+
+double UCelestialMaths::DateTimeToGreenwichMeanSiderealTimeUnwrapped(FDateTime UTCDateTime)
+{
+    double JulianDate = UTCDateTimeToJulianDate(UTCDateTime);
 	double JulianCentury = JulianDateToJulianCenturies(JulianDate);
-	double SideralTime = 280.46061837 + 360.98564736629*(JulianDate-2451545.0) + 0.000387933 * JulianCentury * JulianCentury - 1.0/38710000.0 * JulianCentury * JulianCentury * JulianCentury;
-	return ModPositive(SideralTime, 360.0);
+	return 280.46061837 + 360.98564736629*(JulianDate-2451545.0) + 0.000387933 * JulianCentury * JulianCentury - 1.0/38710000.0 * JulianCentury * JulianCentury * JulianCentury;
 }
 
 double UCelestialMaths::JulianDateToGreenwichMeanSiderealTime(double JulianDate) {
diff --git a/Engine/Plugins/Experimental/CelestialVault/Source/CelestialVault/Private/CelestialVaultDaySequenceActor.cpp b/Engine/Plugins/Experimental/CelestialVault/Source/CelestialVault/Private/CelestialVaultDaySequenceActor.cpp
--- a/Engine/Plugins/Experimental/CelestialVault/Source/CelestialVault/Private/CelestialVaultDaySequenceActor.cpp
+++ b/Engine/Plugins/Experimental/CelestialVault/Source/CelestialVault/Private/CelestialVaultDaySequenceActor.cpp
@@ -80,9 +80,28 @@ ACelestialVaultDaySequenceActor::ACelestialVaultDaySequenceActor(const FObjectIn
 		VolumetricCloudComponent->SetMaterial(VolumetricCloudDefaultMaterialRef.Object);
 	}
 
+	NorthOffsetComponent = CreateDefaultSubobject<USceneComponent>(TEXT("North Offset Transform"));
+	NorthOffsetComponent->SetupAttachment(RootComponent);
+
+	// Add an editor-only compass to visualize the North Offset
+	CompassComponent = CreateOptionalDefaultSubobject<UStaticMeshComponent>(TEXT("Compass"));
+	CompassComponent->SetupAttachment(NorthOffsetComponent);
+	CompassComponent->SetCollisionEnabled(ECollisionEnabled::NoCollision);
+	CompassComponent->SetGenerateOverlapEvents(false);
+	CompassComponent->SetCastShadow(false);
+	CompassComponent->SetAffectDynamicIndirectLighting(false);
+	CompassComponent->SetCanEverAffectNavigation(false);
+	CompassComponent->bIsEditorOnly = true;
+	CompassComponent->SetHiddenInGame(true);
+	if (!IsTemplate())
+	{
+		static ConstructorHelpers::FObjectFinder<UStaticMesh> CompassDefaultMesh(TEXT("/CelestialVault/Editor/SM_Compass.SM_Compass"));
+		CompassComponent->SetStaticMesh(CompassDefaultMesh.Object);
+	}
+	
 	PlanetCenterComponent = CreateDefaultSubobject<USceneComponent>(TEXT("Planet Center Transform"));
-	PlanetCenterComponent->SetupAttachment(RootComponent);
-
+	PlanetCenterComponent->SetupAttachment(NorthOffsetComponent);
+	
 	// Components Attached to Planet Center
 	SkyAtmosphereComponent = CreateDefaultSubobject<USkyAtmosphereComponent>(TEXT("Sky Atmosphere"));
 	SkyAtmosphereComponent->SetupAttachment(PlanetCenterComponent);
@@ -181,6 +200,18 @@ ACelestialVaultDaySequenceActor::ACelestialVaultDaySequenceActor(const FObjectIn
 	}
 }
 
+bool ACelestialVaultDaySequenceActor::IsDaylightSavingsNow() const
+{
+	return UDaylightSavings::IsDaylightSavings(GetDateAndTime(), DaylightSavingsMode, Latitude, DaylightSavingsStart, DaylightSavingsEnd, DaylightSavingsSwitchHour);
+}
+
+bool ACelestialVaultDaySequenceActor::IsDaylightSavingsAtTimeOfDay(double TimeOfDay) const
+{
+	FDateTime Date = GetDate();
+	FTimespan ElapsedTimeOfDay = FTimespan::FromHours(TimeOfDay);
+	return UDaylightSavings::IsDaylightSavings(Date + ElapsedTimeOfDay, DaylightSavingsMode, Latitude, DaylightSavingsStart, DaylightSavingsEnd, DaylightSavingsSwitchHour);
+}
+
 //////// Protected Functions
 
 void ACelestialVaultDaySequenceActor::BeginPlay()
@@ -268,14 +299,9 @@ void ACelestialVaultDaySequenceActor::DebugDrawCallback(UCanvas* Canvas, APlayer
 		FString DebugMessage;
 		DebugMessage += FString("------ Celestial Vault ------\n");
 		DebugMessage += FString("    -- Time -- \n");
-		
 
-#if WITH_EDITORONLY_DATA
-		FDateTime LocalTime(Year, Month, Day, TimeOfDayPreview.Hours, TimeOfDayPreview.Minutes, TimeOfDayPreview.Seconds);
-#else
-		FDateTime LocalTime(Year, Month, Day, 12, 0, 0);
-#endif
-		FDateTime UTCTime = UCelestialMaths::LocalTimeToUTCTime(LocalTime, GMT_TimeZone, bIsDaylightSavings);
+		FDateTime LocalTime = GetDateAndTime();
+		FDateTime UTCTime = UCelestialMaths::LocalTimeToUTCTime(LocalTime, GMT_TimeZone, IsDaylightSavingsNow());
 		double JulianDate = UCelestialMaths::UTCDateTimeToJulianDate(UTCTime);
 		double JulianCenturies = UCelestialMaths::JulianDateToJulianCenturies(JulianDate); 
 		double GMST = UCelestialMaths::DateTimeToGreenwichMeanSiderealTime(UTCTime);
@@ -286,8 +312,9 @@ void ACelestialVaultDaySequenceActor::DebugDrawCallback(UCanvas* Canvas, APlayer
 		DebugMessage += FString::Printf(TEXT("    UTC Time: %s"), *UTCTime.ToFormattedString(TEXT("%d %b %Y, %H:%M:%S"))) + LINE_TERMINATOR;
 		DebugMessage += FString::Printf(TEXT("    Julian Date: %.5f"), JulianDate) + LINE_TERMINATOR;
 		DebugMessage += FString::Printf(TEXT("    Julian Centuries: %.5f"), JulianCenturies) + LINE_TERMINATOR;
-
-		DebugMessage += FString::Printf(TEXT("    GMST at 0:00: %s (%.5f)°"), *UCelestialMaths::Conv_DegreesToHMSString(GMST_AtTOD_0), GMST_AtTOD_0) + LINE_TERMINATOR;
+		DebugMessage += FString::Printf(TEXT("    Daylight Savings: %s"), IsDaylightSavingsNow() ? TEXT("Yes") : TEXT("No") ) + LINE_TERMINATOR;
+		
+	DebugMessage += FString::Printf(TEXT("    GMST at 0:00: %s (%.5f)°"), *UCelestialMaths::Conv_DegreesToHMSString(UCelestialMaths::ModPositive(GMST0_Unwrapped,360)), UCelestialMaths::ModPositive(GMST0_Unwrapped,360)) + LINE_TERMINATOR;
 		DebugMessage += FString::Printf(TEXT("    GMST: %s (%.5f)°"), *UCelestialMaths::Conv_DegreesToHMSString(GMST), GMST) + LINE_TERMINATOR;
 		DebugMessage += FString::Printf(TEXT("    LMST: %s (%.5f)°"), *UCelestialMaths::Conv_DegreesToHMSString(LMST), LMST) + LINE_TERMINATOR;
 
@@ -316,15 +343,22 @@ FDateTime ACelestialVaultDaySequenceActor::GetDate() const
 		return FDateTime::Now().GetDate();
 	}
 
-	if ( Day > FDateTime::DaysInMonth(Year, Month))
+	int32 ClampedMonth = FMath::Clamp(static_cast<int>(Month), 1, 12);
+	if ( Day > FDateTime::DaysInMonth(Year, ClampedMonth))
 	{
-		int32 ClampedDay = FDateTime::DaysInMonth(Year, Month);
-		int32 ClampedMonth = FMath::Clamp(Month, 1, 12);
+		int32 ClampedDay = FDateTime::DaysInMonth(Year, ClampedMonth);
 		UE_LOG(LogCelestialVault, Warning, TEXT("Day value (%d) over the number of days in month - Using %d instead"), Day, ClampedDay);
 		return FDateTime(Year, ClampedMonth, ClampedDay, 0,0,0);
 	}
 		
-	return FDateTime(Year, Month, Day, 0,0,0);
+	return FDateTime(Year, static_cast<int>(Month), Day, 0,0,0);
+}
+
+FDateTime ACelestialVaultDaySequenceActor::GetDateAndTime() const
+{
+	FDateTime Date = GetDate();
+	FTimespan ElapsedTimeOfDay = FTimespan::FromHours(GetTimeOfDay());
+	return Date + ElapsedTimeOfDay;
 }
 
 FStellarBody ACelestialVaultDaySequenceActor::GetSun(double JulianDate) const
@@ -332,7 +366,7 @@ FStellarBody ACelestialVaultDaySequenceActor::GetSun(double JulianDate) const
 	FStellarBody ResultSunInfo = UCelestialMaths::GetSunInformation(JulianDate, Latitude, Longitude);
 
 	// Location
-	FVector SunLocation = UCelestialMaths::RADECToXYZ_RH(ResultSunInfo.RA * 15.0, ResultSunInfo.DEC, 1000.0); // TODO - Add the Sun Distance To Sun 
+	FVector SunLocation = UCelestialMaths::RADECToXYZ_RH(ResultSunInfo.RA * 15.0, ResultSunInfo.DEC, UCelestialMaths::AstronomicalUnitsToMeters(ResultSunInfo.DistanceInAU) * 100); 
 	SunLocation.Y *= -1.0; // Convert to UE Frame by inverting Y
 	ResultSunInfo.UELocalTransform.SetLocation(SunLocation);
 	ResultSunInfo.DirectionTowardEarth = (FVector::ZeroVector - SunLocation).GetSafeNormal();
@@ -424,7 +458,7 @@ bool ACelestialVaultDaySequenceActor::GetClosestPlanetaryBody(FVector ObserverLo
 	BodyTransform = FTransform::Identity;
 	
 	// We query only if we generated the Planet Info  
-	if (!bKeepPlanetsInfos && PlanetsComponent)
+	if (!bKeepPlanetsInfos)
 	{
 		return false;
 	}
@@ -454,14 +488,14 @@ bool ACelestialVaultDaySequenceActor::GetClosestPlanetaryBody(FVector ObserverLo
 	}
 
 	// Check for planets now
-	FTransform PlanetComponentTransform = PlanetsComponent->GetComponentTransform();
-	for (FPlanetaryBody BodyInfo : Planets)
+	if (PlanetsComponent)
 	{
-		if (PlanetsComponent)
+		FTransform PlanetComponentTransform = PlanetsComponent->GetComponentTransform();
+		for (FPlanetaryBody BodyInfo : Planets)
 		{
 			// We need to compute the world Transform for the body, because the celestial vault has rotated and the Body UE Transform is local to the Sky. 
-			FTransform WorldTransform =  BodyInfo.UELocalTransform * PlanetComponentTransform;
-			
+			FTransform WorldTransform = BodyInfo.UELocalTransform * PlanetComponentTransform;
+
 			FVector DirectionToInstance = WorldTransform.GetLocation() - ObserverLocation;
 			DirectionToInstance.Normalize();
 
@@ -510,9 +544,15 @@ bool ACelestialVaultDaySequenceActor::GetPlanetaryBodyByVSOP87Type(EVSOP87BodyTy
 	return false;
 }
 
-double ACelestialVaultDaySequenceActor::GetDayCelestialVaultAngle() const
+double ACelestialVaultDaySequenceActor::GetCelestialVaultAngle(double TimeOfDay) const
 {
-	return GMST_AtTOD_0; 
+	// In order to animate the Celestial vault rotation, we'll use the GMST. But if we modulate the GMST angle, we'll have the 360->0 problem.
+	// We want a value without that problem, but we can't trick using a monotonically increasing value, because when DST changes, it's like if the GMST was going backwards at the DST switch hour. 
+	// Therefore, we keep trace of and "absolute/unwrapped" GMST0, and substract it from the GMST "absolute/unwrapped".
+	// In order to still have the proper GMST rotation at the beginning, we modulate it as a start reference.  
+	FDateTime LocalTime = GetDate() + FTimespan::FromHours(TimeOfDay);
+	FDateTime UTCTime = UCelestialMaths::LocalTimeToUTCTime(LocalTime, GMT_TimeZone, IsDaylightSavingsAtTimeOfDay(TimeOfDay));
+	return UCelestialMaths::ModPositive(GMST0_Unwrapped, 360) + UCelestialMaths::DateTimeToGreenwichMeanSiderealTimeUnwrapped(UTCTime) - GMST0_Unwrapped;
 }
 
 #if WITH_EDITOR
@@ -544,23 +584,34 @@ void ACelestialVaultDaySequenceActor::PostEditChangeProperty(FPropertyChangedEve
 		bRebuildSequence = true;
 	}
 
+	// North Offset
+	if (PropertyName == GET_MEMBER_NAME_CHECKED(ACelestialVaultDaySequenceActor, NorthOffset))
+	{
+		if (NorthOffsetComponent)
+		{
+			NorthOffsetComponent->SetRelativeRotation(FRotator(0, NorthOffset, 0));
+		}
+	}
+		
+	
 	// Time properties
 	if (PropertyName == GET_MEMBER_NAME_CHECKED(ACelestialVaultDaySequenceActor, bUseCurrentDate) ||
 		PropertyName == GET_MEMBER_NAME_CHECKED(ACelestialVaultDaySequenceActor, Year) ||
 		PropertyName == GET_MEMBER_NAME_CHECKED(ACelestialVaultDaySequenceActor, Month) ||
 		PropertyName == GET_MEMBER_NAME_CHECKED(ACelestialVaultDaySequenceActor, Day) ||
 		PropertyName == GET_MEMBER_NAME_CHECKED(ACelestialVaultDaySequenceActor, GMT_TimeZone) ||
-		PropertyName == GET_MEMBER_NAME_CHECKED(ACelestialVaultDaySequenceActor, bIsDaylightSavings) ||
 		PropertyName == GET_MEMBER_NAME_CHECKED(ACelestialVaultDaySequenceActor, Latitude) ||
 		PropertyName == GET_MEMBER_NAME_CHECKED(ACelestialVaultDaySequenceActor, Longitude))
 	{
 		// DateTime has changed, we need to update the Reference CelestialVault Angle.
-		FDateTime LocalDateTimeZero = GetDate();
-		FDateTime UTCDateTimeZero = UCelestialMaths::LocalTimeToUTCTime(LocalDateTimeZero, GMT_TimeZone, bIsDaylightSavings);
-		GMST_AtTOD_0 = UCelestialMaths::DateTimeToGreenwichMeanSiderealTime(UTCDateTimeZero);
-		double CelestialVaultAngle = GetDayCelestialVaultAngle();
-		CelestialVaultComponent->SetRelativeRotation( FRotator(0.0, CelestialVaultAngle, 0.0) );
+		FDateTime LocalTime0 = GetDate();
+		FDateTime UTCTime0 = UCelestialMaths::LocalTimeToUTCTime(LocalTime0, GMT_TimeZone, IsDaylightSavingsAtTimeOfDay(0));
+		GMST0_Unwrapped = UCelestialMaths::DateTimeToGreenwichMeanSiderealTimeUnwrapped(UTCTime0);
+		
+		CelestialVaultComponent->SetRelativeRotation( FRotator(0.0, UCelestialMaths::ModPositive(GMST0_Unwrapped, 360), 0.0) );
 
+		UpdatePlanetCenter();
+		
 		bBuildAll = true;
 		bRebuildSequence = true;
 	}
@@ -609,6 +660,37 @@ void ACelestialVaultDaySequenceActor::PostEditChangeProperty(FPropertyChangedEve
 
 	Super::PostEditChangeProperty(PropertyChangedEvent);
 }
+
+
+void ACelestialVaultDaySequenceActor::PostEditChangeChainProperty(FPropertyChangedChainEvent& PropertyChangedEvent)
+{
+	Super::PostEditChangeChainProperty(PropertyChangedEvent);
+
+	// Check the Validity of Daylight Savings Dates
+	if (FProperty* MemberProperty = PropertyChangedEvent.PropertyChain.GetHead()->GetValue()) // The Struct property that has changed is the Head of the Chain
+	{
+		// Is it a FDaylightSavingsRule property that changed? 
+		FStructProperty* StructProperty = CastField<FStructProperty>(MemberProperty);
+		if (StructProperty && StructProperty->Struct == FDaylightSavingsRule::StaticStruct())
+		{
+			// Now, check which property of this DaylightSavingsRule has changed
+			if (FProperty* InnerProperty = PropertyChangedEvent.PropertyChain.GetTail()->GetValue())
+			{
+				const FName PropName = InnerProperty->GetFName();
+
+				// Validate only if Fixed-Day related properties changed
+				if (PropName == GET_MEMBER_NAME_CHECKED(FDaylightSavingsRule, DayOfMonth) ||
+					PropName == GET_MEMBER_NAME_CHECKED(FDaylightSavingsRule, Month))
+				{
+					if (FDaylightSavingsRule* DaylightSavingsRule = StructProperty->ContainerPtrToValuePtr<FDaylightSavingsRule>(this))
+					{
+						DaylightSavingsRule->ClampToValidDate(Year);
+					}
+				}
+			}
+		}
+	}
+}
 #endif
 
 //////// Private  Functions
@@ -744,9 +826,9 @@ void ACelestialVaultDaySequenceActor::InitPlanetaryBodies()
 	TRACE_CPUPROFILER_EVENT_SCOPE(ACelestialVaultDaySequenceActor::InitPlanetaryBodies);
 	
 	// Get the proper Julian day, at Midnight... The Daysequence will rotate the sky vault later 
-	FDateTime LocalTimeMidnight = GetDate();
-	FDateTime UTCTime = UCelestialMaths::LocalTimeToUTCTime(LocalTimeMidnight, GMT_TimeZone, bIsDaylightSavings);
-	double JulianDay = UCelestialMaths::UTCDateTimeToJulianDate(UTCTime);
+	FDateTime LocalTime0 = GetDate();
+	FDateTime UTCTime0 = UCelestialMaths::LocalTimeToUTCTime(LocalTime0, GMT_TimeZone, IsDaylightSavingsNow());
+	double JulianDay = UCelestialMaths::UTCDateTimeToJulianDate(UTCTime0);
 
 	// Init the sun
 	Sun = GetSun(JulianDay);
@@ -837,6 +919,17 @@ FPlanetaryBody ACelestialVaultDaySequenceActor::GetPlanetaryBodyInfo(const FPlan
 	return BodyInfo;
 }
 
+void ACelestialVaultDaySequenceActor::UpdatePlanetCenter()
+{
+	// Get the transformation in proper UE Frame units, and set it to the Component 
+	PlanetCenterTransform = UCelestialMaths::GetEarthCenterTransformUEFrame(Latitude, Longitude, 0);
+	PlanetCenterComponent->SetWorldTransform(PlanetCenterTransform);
+
+	// Adjust the SkyAtmosphere to the new radius
+	double Radius = PlanetCenterTransform.GetLocation().Length();
+	SkyAtmosphereComponent->SetBottomRadius(Radius / UE_ONE_KILOMETER);
+}
+
 void ACelestialVaultDaySequenceActor::RebuildAll()
 {
 	if (const UWorld* World = GetWorld())
@@ -845,19 +938,16 @@ void ACelestialVaultDaySequenceActor::RebuildAll()
 		SetActorTransform(FTransform::Identity);
 		
 		// DateTime has changed, we need to update the Reference CelestialVault Angle.
-		FDateTime LocalDateTimeZero = GetDate();
-		FDateTime UTCDateTimeZero = UCelestialMaths::LocalTimeToUTCTime(LocalDateTimeZero, GMT_TimeZone, bIsDaylightSavings);
-		GMST_AtTOD_0 = UCelestialMaths::DateTimeToGreenwichMeanSiderealTime(UTCDateTimeZero);
+		FDateTime LocalTime0 = GetDate();
+		FDateTime UTCTime0 = UCelestialMaths::LocalTimeToUTCTime(LocalTime0, GMT_TimeZone, IsDaylightSavingsAtTimeOfDay(0));
+		GMST0_Unwrapped = UCelestialMaths::DateTimeToGreenwichMeanSiderealTimeUnwrapped(UTCTime0);
+		
+		CelestialVaultComponent->SetRelativeRotation( FRotator(0.0, UCelestialMaths::ModPositive(GMST0_Unwrapped, 360.0), 0.0) );
 
-		// Get the transformation in proper UE Frame units, and set it to the Component 
-		PlanetCenterTransform = UCelestialMaths::GetEarthCenterTransformUEFrame(Latitude, Longitude, 0);
-		PlanetCenterComponent->SetWorldTransform(PlanetCenterTransform);
+		UpdatePlanetCenter();
 
 		// Rebuild Sky and Sequence
 		InitStars();
 		InitPlanetaryBodies();
-
-		double CelestialVaultAngle = GetDayCelestialVaultAngle();
-		CelestialVaultComponent->SetRelativeRotation( FRotator(0.0, CelestialVaultAngle, 0.0) );
 	}
 }
diff --git a/Engine/Plugins/Experimental/CelestialVault/Source/CelestialVault/Private/CelestialVaultSequence.cpp b/Engine/Plugins/Experimental/CelestialVault/Source/CelestialVault/Private/CelestialVaultSequence.cpp
--- a/Engine/Plugins/Experimental/CelestialVault/Source/CelestialVault/Private/CelestialVaultSequence.cpp
+++ b/Engine/Plugins/Experimental/CelestialVault/Source/CelestialVault/Private/CelestialVaultSequence.cpp
@@ -30,14 +30,48 @@ void FCelestialVaultSequence::BuildSequence(UProceduralDaySequenceBuilder* InBui
 	if (USceneComponent* CelestialVaultComponent = EarthDaySequenceActor->CelestialVaultComponent)
 	{
 		InBuilder->SetActiveBoundObject(CelestialVaultComponent);
+
+		const double NormalizedTimeIncrement = 1.0 / FMath::Max(static_cast<signed>(KeyCount) - 1 , 1);
 		
-		const double NormalizedTimeIncrement = 1.0 / FMath::Max(static_cast<signed>(KeyCount) - 1, 1);
-		for (unsigned int Key = 0; Key < KeyCount; ++Key)
+		
+		// Check for DaylightSavingsEnd year. 
+		int32 EndTimeYear = EarthDaySequenceActor->Year;
+		if ( EarthDaySequenceActor->DaylightSavingsMode == EDaylightSavingsMode::Southern ||
+			(EarthDaySequenceActor->DaylightSavingsMode == EDaylightSavingsMode::Latitude && EarthDaySequenceActor->Latitude < 0))
 		{
-			double CelestialVaultAngleAtMidnight = EarthDaySequenceActor->GetDayCelestialVaultAngle(); 
+			EndTimeYear++;
+		}
 
+		double DSTKeyStart = -1.0;
+		double DSTKeyEnd = -1.0;
+		// If we changed DST this day, add additional keys around the DST Switch hour.
+		if (EarthDaySequenceActor->GetDate() == UDaylightSavings::ToDate(EarthDaySequenceActor->DaylightSavingsStart, EarthDaySequenceActor->Year) ||
+			EarthDaySequenceActor->GetDate() == UDaylightSavings::ToDate(EarthDaySequenceActor->DaylightSavingsEnd, EndTimeYear))
+		{
+			double KeySpanSeconds = 30;
+			double KeySpanNormalized = KeySpanSeconds / (24.0*3600.0); 
+			double SwitchHourNormalized = EarthDaySequenceActor->DaylightSavingsSwitchHour / 24.0; 
+			
+			DSTKeyStart = FMath::Clamp(SwitchHourNormalized - 0.5 * KeySpanNormalized, 0.0, 1.0);
+			DSTKeyEnd = FMath::Clamp(SwitchHourNormalized + 0.5 * KeySpanNormalized, 0.0, 1.0);
+		} 
+				
+		for (unsigned int Key = 0; Key < KeyCount; ++Key)
+		{
 			const double KeyTime = Key * NormalizedTimeIncrement;
-			InBuilder->AddRotationKey(KeyTime, FRotator(0.0, CelestialVaultAngleAtMidnight + KeyTime * 360.0, 0.0));
+
+			if (KeyTime > DSTKeyStart && KeyTime < DSTKeyEnd )
+			{
+				continue;
+			}
+			InBuilder->AddRotationKey(KeyTime, FRotator(0.0, EarthDaySequenceActor->GetCelestialVaultAngle(KeyTime * 24.0), 0.0), RCIM_Linear);
+		}
+
+		// If we changed DST this day, add additional keys around the DST Switch hour.
+		if (DSTKeyStart != -1.0 && DSTKeyEnd != -1.0)
+		{
+			InBuilder->AddRotationKey(DSTKeyStart, FRotator(0.0, EarthDaySequenceActor->GetCelestialVaultAngle(DSTKeyStart * 24.0), 0.0), RCIM_Linear);
+			InBuilder->AddRotationKey(DSTKeyEnd, FRotator(0.0, EarthDaySequenceActor->GetCelestialVaultAngle(DSTKeyEnd * 24.0), 0.0), RCIM_Linear);
 		}
 	}
 
@@ -53,8 +87,7 @@ void FCelestialVaultSequence::BuildSequence(UProceduralDaySequenceBuilder* InBui
 			const double KeyTime = Key * NormalizedTimeIncrement;
 
 			FDateTime KeyDateTime = LocalTimeMidnight + FTimespan::FromSeconds(KeyTime * 24.0 * 60.0 * 60.0);
-			FDateTime UTCTime = UCelestialMaths::LocalTimeToUTCTime(KeyDateTime, EarthDaySequenceActor->GMT_TimeZone,
-				EarthDaySequenceActor->bIsDaylightSavings);
+			FDateTime UTCTime = UCelestialMaths::LocalTimeToUTCTime(KeyDateTime, EarthDaySequenceActor->GMT_TimeZone, EarthDaySequenceActor->IsDaylightSavingsAtTimeOfDay(KeyTime * 24.0));
 			double JulianDay = UCelestialMaths::UTCDateTimeToJulianDate(UTCTime);
 
 			FPlanetaryBody  MoonInfo = EarthDaySequenceActor->GetMoon(JulianDay);
@@ -86,8 +119,7 @@ void FCelestialVaultSequence::BuildSequence(UProceduralDaySequenceBuilder* InBui
 			const double KeyTime = Key * NormalizedTimeIncrement;
 
 			FDateTime KeyDateTime = LocalTimeMidnight + FTimespan::FromSeconds(KeyTime * 24 * 60 * 60);
-			FDateTime UTCTime = UCelestialMaths::LocalTimeToUTCTime(KeyDateTime, EarthDaySequenceActor->GMT_TimeZone,
-				EarthDaySequenceActor->bIsDaylightSavings);
+			FDateTime UTCTime = UCelestialMaths::LocalTimeToUTCTime(KeyDateTime, EarthDaySequenceActor->GMT_TimeZone, EarthDaySequenceActor->IsDaylightSavingsAtTimeOfDay(KeyTime * 24.0));
 			double JulianDay = UCelestialMaths::UTCDateTimeToJulianDate(UTCTime);
 
 			FStellarBody SunInfo = EarthDaySequenceActor->GetSun(JulianDay);
diff --git a/Engine/Plugins/Experimental/CelestialVault/Source/CelestialVault/Private/DaylightSavings.cpp b/Engine/Plugins/Experimental/CelestialVault/Source/CelestialVault/Private/DaylightSavings.cpp
new file mode 100644
--- /dev/null
+++ b/Engine/Plugins/Experimental/CelestialVault/Source/CelestialVault/Private/DaylightSavings.cpp
@@ -0,0 +1,165 @@
+﻿// Copyright Epic Games, Inc. All Rights Reserved.
+
+
+
+#include "DaylightSavings.h"
+
+FDateTime UDaylightSavings::ToDate(const FDaylightSavingsRule& DaylightSavingsRule, int Year)
+{
+	FDateTime Result;
+
+	switch (DaylightSavingsRule.Rule)
+	{
+	case EDaylightSavingRuleKind::NthDay:
+		{
+			int DayInMonth = NthWeekdayInMonth(Year, DaylightSavingsRule.Month, DaylightSavingsRule.WeekDay, DaylightSavingsRule.Nth);
+			if (DayInMonth == 0 )
+			{
+				Result = FDateTime::MinValue();
+			}
+			else
+			{
+				Result = FDateTime(Year,FMath::Clamp(static_cast<int>(DaylightSavingsRule.Month), 1, 12), DayInMonth );	
+			}
+		}
+		break;
+	case EDaylightSavingRuleKind::LastDay:
+		Result = FDateTime(Year, FMath::Clamp(static_cast<int>(DaylightSavingsRule.Month), 1, 12), LastWeekdayInMonth(Year, DaylightSavingsRule.Month, DaylightSavingsRule.WeekDay)); 
+		break;
+	case EDaylightSavingRuleKind::FixedDate:
+		Result = FDateTime(Year, FMath::Clamp(static_cast<int>(DaylightSavingsRule.Month), 1, 12), FMath::Clamp(DaylightSavingsRule.DayOfMonth, 1 , FDateTime::DaysInMonth(Year, FMath::Clamp(static_cast<int>(DaylightSavingsRule.Month), 1, 12))));
+		break;
+	}
+
+	return Result;
+}
+
+FString UDaylightSavings::ToString(const FDaylightSavingsRule& DaylightSavingsRule)
+{
+	FString Result = FString();
+	switch (DaylightSavingsRule.Rule)
+	{
+	case EDaylightSavingRuleKind::NthDay:
+		Result = FString::Printf(TEXT("%d%s %s of %s"), DaylightSavingsRule.Nth, *GetOrdinalSuffix(DaylightSavingsRule.Nth), *GetDayString(DaylightSavingsRule.WeekDay), *GetMonthString(DaylightSavingsRule.Month));
+		break;
+	case EDaylightSavingRuleKind::LastDay:
+		Result =  FString::Printf(TEXT("Last %s of %s"), *GetDayString(DaylightSavingsRule.WeekDay), *GetMonthString(DaylightSavingsRule.Month));
+		break;
+	case EDaylightSavingRuleKind::FixedDate:
+		Result =  FString::Printf(TEXT("%d%s of %s"), DaylightSavingsRule.DayOfMonth, *GetOrdinalSuffix(DaylightSavingsRule.DayOfMonth), *GetMonthString(DaylightSavingsRule.Month));
+		break;
+	}
+
+	return Result;
+}
+
+EWeekDay UDaylightSavings::GetWeekDay(FDateTime Date)
+{
+	// Uses Sakamoto’s algorithm internally, which returns 0=Sunday..6=Saturday.
+	static int month_table[12] = { 0, 3, 2, 5, 0, 3, 5, 1, 4, 6, 2, 4 }; 
+	int Year = Date.GetYear();
+	if (Date.GetMonth() < 3)
+	{
+		Year -= 1;
+	}
+	
+	int Day = ( Year + Year / 4 - Year / 100 + Year / 400 + month_table[Date.GetMonth() - 1] + Date.GetDay() ) % 7;
+
+	// We want our Week to start on Monday (BP enums needs to start at 0), so Mapping from Sunday-based to Monday-based:
+	return static_cast<EWeekDay>((Day+6)%7);
+}
+
+int UDaylightSavings::NthWeekdayInMonth(int Year, EMonth Month, EWeekDay Weekday, int Nth)
+{
+	if (Nth <= 0)
+	{
+		return 0;
+	}
+	
+	EWeekDay FirstWeekDayOfMonth = GetWeekDay(FDateTime(Year,FMath::Clamp(static_cast<int>(Month), 1, 12), 1));
+	int DayOffset = (static_cast<int>(Weekday) - static_cast<int>(FirstWeekDayOfMonth) + 7) % 7;
+
+	int Day = 1 + DayOffset + 7 * (Nth - 1);
+	if (Day > FDateTime::DaysInMonth(Year, FMath::Clamp(static_cast<int>(Month), 1, 12)))
+	{
+		return 0; // does not exist
+	}
+
+	return Day;
+}
+
+int UDaylightSavings::LastWeekdayInMonth(int Year, EMonth Month, EWeekDay Weekday)
+{
+	int LastDay = FDateTime::DaysInMonth(Year, FMath::Clamp(static_cast<int>(Month), 1, 12)); 
+	EWeekDay LastWeekdayOfMonth = GetWeekDay(FDateTime(Year, FMath::Clamp(static_cast<int>(Month), 1, 12), LastDay));
+
+	int DayOffsetBackwards = (static_cast<int>(LastWeekdayOfMonth) - static_cast<int>(Weekday) + 7) % 7;
+	return LastDay - DayOffsetBackwards;
+}
+
+FString UDaylightSavings::GetMonthString(EMonth Month)
+{
+	FString Result = FString();
+	if (UEnum* EnumClass = StaticEnum<EMonth>())
+	{
+		Result = EnumClass->GetNameStringByValue(static_cast<int64>(Month));
+	}
+
+	return Result;
+}
+
+FString UDaylightSavings::GetDayString(EWeekDay Day)
+{
+	FString Result = FString();
+	if (UEnum* EnumClass = StaticEnum<EWeekDay>())
+	{
+		Result = EnumClass->GetNameStringByValue(static_cast<int64>(Day));
+	}
+
+	return Result;
+}
+
+FString UDaylightSavings::GetOrdinalSuffix(int32 Number)
+{
+	int32 Tens = Number % 100;
+
+	// Exceptions : 11, 12, 13 -> "th"
+	if (Tens >= 11 && Tens <= 13)
+	{
+		return TEXT("th");
+	}
+
+	// Otherwise look at the unit
+	switch (Number % 10)
+	{
+	case 1:
+		return TEXT("st");
+	case 2:
+		return TEXT("nd");
+	case 3:
+		return TEXT("rd");
+	default:
+		return TEXT("th");
+	}
+}
+
+bool UDaylightSavings::IsDaylightSavings(FDateTime DateTime, EDaylightSavingsMode Mode, double Latitude, FDaylightSavingsRule StartDay, FDaylightSavingsRule EndDay, int SwitchHour)
+{
+	if (Mode == EDaylightSavingsMode::None)
+	{
+		return false;
+	}
+
+	int32 EndTimeYear = DateTime.GetYear();
+
+	if (Mode == EDaylightSavingsMode::Southern ||
+		(Mode == EDaylightSavingsMode::Latitude && Latitude < 0))
+	{
+		EndTimeYear++;
+	}
+	
+	FDateTime StartTime = ToDate(StartDay, DateTime.GetYear()) + FTimespan::FromHours(SwitchHour);
+	FDateTime EndTime = ToDate(EndDay,  EndTimeYear) + FTimespan::FromHours(SwitchHour);
+
+	return (DateTime >= StartTime && DateTime <= EndTime);
+}
diff --git a/Engine/Plugins/Experimental/CelestialVault/Source/CelestialVault/Public/CelestialMaths.h b/Engine/Plugins/Experimental/CelestialVault/Source/CelestialVault/Public/CelestialMaths.h
--- a/Engine/Plugins/Experimental/CelestialVault/Source/CelestialVault/Public/CelestialMaths.h
+++ b/Engine/Plugins/Experimental/CelestialVault/Source/CelestialVault/Public/CelestialMaths.h
@@ -124,6 +124,15 @@ public:
     UFUNCTION(BlueprintCallable, Category="Celestial Vault|Time")
     static double DateTimeToGreenwichMeanSiderealTime(FDateTime UTCDateTime);
 
+    /**
+     * Return the Greenwich Mean Sidereal Time (GMST) for a specific DateTime, in Degrees, without doing any ModPositive, so unwrapped
+     * 
+     * By definition, the provided DateTime has to be the DateTime at the Greenwitch Meridian, so it's a UTC DateTime
+     */
+    UFUNCTION(BlueprintCallable, Category="Celestial Vault|Time")
+    static double DateTimeToGreenwichMeanSiderealTimeUnwrapped(FDateTime UTCDateTime);
+
+
     /** Return the Greenwich Mean Sidereal Time (GMST) for a specific Julian Date, In Degrees */
     UFUNCTION(BlueprintCallable, Category="Celestial Vault|Time")
     static double JulianDateToGreenwichMeanSiderealTime(double JulianDate);
diff --git a/Engine/Plugins/Experimental/CelestialVault/Source/CelestialVault/Public/CelestialVaultDaySequenceActor.h b/Engine/Plugins/Experimental/CelestialVault/Source/CelestialVault/Public/CelestialVaultDaySequenceActor.h
--- a/Engine/Plugins/Experimental/CelestialVault/Source/CelestialVault/Public/CelestialVaultDaySequenceActor.h
+++ b/Engine/Plugins/Experimental/CelestialVault/Source/CelestialVault/Public/CelestialVaultDaySequenceActor.h
@@ -5,6 +5,7 @@
 #include "CoreMinimal.h"
 #include "CelestialDataTypes.h"
 #include "CelestialInputDataTypes.h"
+#include "DaylightSavings.h"
 #include "DaySequenceActor.h"
 #include "CelestialVaultDaySequenceActor.generated.h"
 
@@ -46,6 +47,12 @@ public: // Properties
 	UPROPERTY(VisibleDefaultsOnly, BlueprintReadOnly, Category= "Celestial Vault", meta = (AllowPrivateAccess = "true"))
 	TObjectPtr<USceneComponent> PlanetCenterComponent;
 
+	UPROPERTY(VisibleDefaultsOnly, BlueprintReadOnly, Category= "Celestial Vault", meta = (AllowPrivateAccess = "true"))
+	TObjectPtr<USceneComponent> NorthOffsetComponent;
+
+	UPROPERTY(VisibleDefaultsOnly, BlueprintReadOnly, Category= "Celestial Vault", meta = (AllowPrivateAccess = "true"))
+	TObjectPtr<UStaticMeshComponent> CompassComponent;
+
 	UPROPERTY(VisibleDefaultsOnly, BlueprintReadOnly, Category= "Celestial Vault", meta = (AllowPrivateAccess = "true"))
 	TObjectPtr<USceneComponent> CelestialVaultComponent;
 
@@ -87,43 +94,81 @@ public: // Properties
 #pragma region DateTime related Properties
 
 	/* If true, ignore the Year Month Day value and use the current system Date */ 
-	UPROPERTY(EditAnywhere, BlueprintReadWrite, Category="Celestial Vault|Date & Location")
+	UPROPERTY(EditAnywhere, BlueprintReadWrite, Category="Date")
 	bool bUseCurrentDate = false;
 
 	/** Current Year*/ 
-	UPROPERTY(EditAnywhere, BlueprintReadWrite, Category="Celestial Vault|Date & Location", meta = (EditConditionHides, EditCondition = "bUseCurrentDate==false", ClampMin = "1", ClampMax = "9999"))
+	UPROPERTY(EditAnywhere, BlueprintReadWrite, Category="Date", meta = (EditConditionHides, EditCondition = "!bUseCurrentDate", ClampMin = "1", ClampMax = "9999"))
 	int Year = 2025;
 
 	/** Current Month*/ 
-	UPROPERTY(EditAnywhere, BlueprintReadWrite, Category="Celestial Vault|Date & Location", meta = (EditConditionHides, EditCondition = "bUseCurrentDate==false", ClampMin = "1", ClampMax = "12"))
-	int Month = 1;
+	UPROPERTY(EditAnywhere, BlueprintReadWrite, Category="Date", meta = (EditConditionHides, EditCondition = "!bUseCurrentDate", ClampMin = "1", ClampMax = "12"))
+	EMonth Month = EMonth::January;
 
 	/** Current Day*/ 
-	UPROPERTY(EditAnywhere, BlueprintReadWrite, Category="Celestial Vault|Date & Location", meta = (EditConditionHides, EditCondition = "bUseCurrentDate==false", ClampMin = "1", ClampMax = "31"))
+	UPROPERTY(EditAnywhere, BlueprintReadWrite, Category="Date", meta = (EditConditionHides, EditCondition = "!bUseCurrentDate", ClampMin = "1", ClampMax = "31"))
 	int Day = 7;
 
+#pragma endregion
+
+#pragma region Daylight Saving related Properties
+    
+	/** Define how the Daylight Savings are taken into consideration */ 
+	UPROPERTY(EditAnywhere, BlueprintReadWrite, Category = "Daylight Savings")
+	EDaylightSavingsMode DaylightSavingsMode = EDaylightSavingsMode::Latitude;
+	
+	/** Rule that applies to determine the Daylight Saving stating day */ 
+	UPROPERTY(EditAnywhere, BlueprintReadWrite, Category="Daylight Savings", meta=(EditConditionHides, EditCondition = "DaylightSavingsMode!=EDaylightSavingsMode::None", DisplayName = "Starts On"))
+	FDaylightSavingsRule DaylightSavingsStart = FDaylightSavingsRule(EDaylightSavingRuleKind::NthDay, 2, EWeekDay::Sunday, 1, EMonth::March);  // Default to 2nd Sunday of March (US/Canada)
+
+	/** Rule that applies to determine the Daylight Saving ending day */ 
+	UPROPERTY(EditAnywhere, BlueprintReadWrite, Category="Daylight Savings", meta=(EditConditionHides, EditCondition = "DaylightSavingsMode!=EDaylightSavingsMode::None", DisplayName = "End On"))
+	FDaylightSavingsRule DaylightSavingsEnd = FDaylightSavingsRule(EDaylightSavingRuleKind::NthDay, 1, EWeekDay::Sunday, 1, EMonth::November); // Default to 1st Sunday of November (US/Canada)
+
+	/** Time of day when the new Daylight savings actually apply */ 
+	UPROPERTY(EditAnywhere, BlueprintReadWrite, Category="Daylight Savings",  meta = (EditConditionHides, EditCondition = "DaylightSavingsMode!=EDaylightSavingsMode::None", DisplayName = "Switch Hour", ClampMin=0, ClampMax=23))
+	int DaylightSavingsSwitchHour = 2;
+
+	/** Return true if the current Date and Time of the Celestial Vault Actor is inside the Daylight savings period */ 
+	UFUNCTION(BlueprintCallable, Category = "Daylight Savings")
+	bool IsDaylightSavingsNow() const;
+
+	/** Return true if the current Date Celestial Vault Actor with an arbitrary Time of day is inside the Daylight savings period */
+	UFUNCTION(BlueprintCallable, Category = "Daylight Savings")
+	bool IsDaylightSavingsAtTimeOfDay(double TimeOfDay) const;
+
+#pragma endregion
+
+#pragma region Location related Properties
+
 	/** Current Time Zone*/ 
-	UPROPERTY(EditAnywhere, BlueprintReadWrite, Category="Celestial Vault|Date & Location", meta = (ClampMin = "-11", ClampMax = "14"))
+	UPROPERTY(EditAnywhere, BlueprintReadWrite, Category="Location", meta = (ClampMin = "-11", ClampMax = "14"))
 	double GMT_TimeZone = -5.0;
-
-	/** Set to true if your current date is during the DST period */ 
-	UPROPERTY(EditAnywhere, BlueprintReadWrite, Category = "Celestial Vault|Date & Location")
-	bool bIsDaylightSavings = false;
 	
 	/** Latitude of Level Origin on planet */
-	UPROPERTY(EditAnywhere, BlueprintReadWrite, Category = "Celestial Vault|Date & Location", meta = (ForceUnits="Degrees", ClampMin = "-90", ClampMax = "90"))
+	UPROPERTY(EditAnywhere, BlueprintReadWrite, Category = "Location", meta = (ForceUnits="Degrees", ClampMin = "-90", ClampMax = "90"))
 	double Latitude = 45.0;
 
 	/** Longitude of Level Origin on planet */
-	UPROPERTY(EditAnywhere, BlueprintReadWrite, Category = "Celestial Vault|Date & Location", meta = (ForceUnits="Degrees", ClampMin = "-180", ClampMax = "180"))
+	UPROPERTY(EditAnywhere, BlueprintReadWrite, Category = "Location", meta = (ForceUnits="Degrees", ClampMin = "-180", ClampMax = "180"))
 	double Longitude = -73.0;
+
+	/**
+	 * North Offset
+	 *
+	 * In most cases, this value is supposed to be 0.
+	 * But in some architecture cases, where the input geometry is aligned to the building main directions and not to the East/North axes,
+	 * this value helps rotate artificially the celestial vault to accomodate for this offset.  
+	 */
+	UPROPERTY(EditAnywhere, BlueprintReadOnly, Category = "Location", meta = (ForceUnits="Degrees", ClampMin = "0", ClampMax = "360"))
+	double NorthOffset = 0.0;
 	
 	// Greenwich Mean Sidereal Time at corresponding to a 0 Time of Day (midnight in the morning) for the selected Date. 
-	UPROPERTY(BlueprintReadOnly, VisibleAnywhere, Category="Celestial Vault|Date & Location")
-	double GMST_AtTOD_0 = 0.0;
+	UPROPERTY(BlueprintReadOnly, Category="Location")
+	double GMST0_Unwrapped = 0.0;
 
 	// Transform to apply to the planet to have it located tangent to the Origin
-	UPROPERTY(BlueprintReadOnly, VisibleAnywhere, Category="Celestial Vault|Date & Location")
+	UPROPERTY(BlueprintReadOnly, VisibleAnywhere, Category="Location")
 	FTransform PlanetCenterTransform;
 	
 #pragma endregion
@@ -151,23 +196,23 @@ public: // Properties
 #pragma region Stars related Properties
 	
 	/** A Datatable containing a Celestial Star Catalog data */
-	UPROPERTY(EditAnywhere, BlueprintReadWrite, Category="Celestial Vault|Stars")
+	UPROPERTY(EditAnywhere, BlueprintReadWrite, Category="Stars")
 	TObjectPtr<UDataTable> CelestialStarCatalog = nullptr;
 
 	/** A Datatable containing a Fictional Star Catalog data */
-	UPROPERTY(EditAnywhere, BlueprintReadWrite, Category="Celestial Vault|Stars")
+	UPROPERTY(EditAnywhere, BlueprintReadWrite, Category="Stars")
 	TObjectPtr<UDataTable> FictionalStarCatalog = nullptr;
 	
 	/** All stars from the catalog with a Magnitude dimmer than this threshold won't be generated - Usually 6 is the naked eye visibility limit*/ 
-	UPROPERTY(EditAnywhere, BlueprintReadWrite, Category="Celestial Vault|Stars", meta = (ClampMin = "-30", ClampMax = "30"))
+	UPROPERTY(EditAnywhere, BlueprintReadWrite, Category="Stars", meta = (ClampMin = "-30", ClampMax = "30"))
 	float MaxVisibleMagnitude = 6.0f;
 
 	/** If true, the Stars information will be kept in memory and queryable at runtime */ 
-	UPROPERTY(EditAnywhere, BlueprintReadWrite, Category="Celestial Vault|Stars")
+	UPROPERTY(EditAnywhere, BlueprintReadWrite, Category="Stars")
 	bool bKeepStarsInfo = false;
 	
 	/** Array of the created Stars information - Only populated if KeepStarsInfo is true */ 
-	UPROPERTY(BlueprintReadOnly, Category="Celestial Vault|Stars")
+	UPROPERTY(BlueprintReadOnly, Category="Stars")
 	TArray<FStellarBody> Stars;
 	
 #pragma endregion
@@ -175,15 +220,15 @@ public: // Properties
 #pragma region Planets related Properties
 
 	/** The Data Catalog containing all Planets data */
-	UPROPERTY(EditAnywhere, BlueprintReadWrite, Category="Celestial Vault|Planets")
+	UPROPERTY(EditAnywhere, BlueprintReadWrite, Category="Planets")
 	TObjectPtr<UDataTable> PlanetsCatalog = nullptr;
 
 	/** Factor to artificially increase the Planetary bodies size */ 
-	UPROPERTY(EditAnywhere, BlueprintReadWrite, Category="Celestial Vault|Planets", meta = (ClampMin = "0"))
+	UPROPERTY(EditAnywhere, BlueprintReadWrite, Category="Planets", meta = (ClampMin = "0"))
 	float PlanetsScale = 1.0f;
 
 	/** If true, the Stars information will be kept in memory and queryable at runtime */ 
-	UPROPERTY(EditAnywhere, BlueprintReadWrite, Category="Celestial Vault|Planets")
+	UPROPERTY(EditAnywhere, BlueprintReadWrite, Category="Planets")
 	bool bKeepPlanetsInfos = false;
 	
 	/** Array of the created planetary bodies, with all their computed information - Only populated if KeepPlanetsInfos is true */ 
@@ -195,19 +240,19 @@ public: // Properties
 #pragma region Moon / Sun related Properties
 
 	/** Factor to artificially increase the Moon size */
-	UPROPERTY(EditAnywhere, BlueprintReadWrite, Category="Celestial Vault|Moons", meta = (ClampMin = "0"))
+	UPROPERTY(EditAnywhere, BlueprintReadWrite, Category="Moons", meta = (ClampMin = "0"))
 	float MoonScale = 2.0f;
 
 	/** Celestial Info for the Moon, at the beginning of the Day */  
-	UPROPERTY(BlueprintReadOnly, Category="Celestial Vault|Moons")
+	UPROPERTY(BlueprintReadOnly, Category="Moons")
 	FPlanetaryBody Moon;
 
 	/** If true, the moon Age (Phase) and location can be overriden */  
-	UPROPERTY(EditAnywhere, Category="Celestial Vault|Moons")
+	UPROPERTY(EditAnywhere, Category="Moons")
 	bool bManualControl = false;
 	
 	/** Lunar age. 0 = New Moon, 0.25 = First quarter, 0.5 = Full Moon, 1 = Next New Mo */  
-	UPROPERTY(EditAnywhere, Category="Celestial Vault|Moons", meta = (EditConditionHides, EditCondition = "bManualControl==true", ClampMin = "0", ClampMax = "1"))
+	UPROPERTY(EditAnywhere, Category="Moons", meta = (EditConditionHides, EditCondition = "bManualControl==true", ClampMin = "0", ClampMax = "1"))
 	float MoonAge = 0.2f;
 	
 	/** When faking the moon location, we need to give a location relative to the sun
@@ -218,16 +263,16 @@ public: // Properties
 	 *   (if HoursBehindSun < 12, we are ~1st quarter, so Age should be < 0.5, and if 12< < 24 age should be > 0.5) 
 	 * Any combination is possible, but manual control can lead to inconsistencies. 
 	 */  
-	UPROPERTY(EditAnywhere, Category="Celestial Vault|Moons", meta = (EditConditionHides, DisplayName="Moon Offset to Sun's Right Ascension", EditCondition = "bManualControl==true", ForceUnits="Hours", ClampMin = "0", ClampMax = "24"))
+	UPROPERTY(EditAnywhere, Category="Moons", meta = (EditConditionHides, DisplayName="Moon Offset to Sun's Right Ascension", EditCondition = "bManualControl==true", ForceUnits="Hours", ClampMin = "0", ClampMax = "24"))
 	float MoonOffset_RA = 12.0f;
 
 	/** When faking the moon location, we need to give a location relative to the sun
 	 * This is a way to control this "Vertically" using an offset in Declination. */ 
-	UPROPERTY(EditAnywhere, Category="Celestial Vault|Moons", meta = (EditConditionHides, DisplayName="Moon Offset to Sun's Declination", EditCondition = "bManualControl==true", ForceUnits="Degrees", ClampMin = "-45", ClampMax = "45"))
+	UPROPERTY(EditAnywhere, Category="Moons", meta = (EditConditionHides, DisplayName="Moon Offset to Sun's Declination", EditCondition = "bManualControl==true", ForceUnits="Degrees", ClampMin = "-45", ClampMax = "45"))
 	float MoonOffset_DEC = 15.0f;
 
 	/** Celestial Info for the Sun, at the beginning of the Day */  
-	UPROPERTY(BlueprintReadOnly, Category="Celestial Vault|Sun")
+	UPROPERTY(BlueprintReadOnly, Category="Sun")
 	FStellarBody Sun; 
 
 	/** Base Sun Intensity
@@ -237,7 +282,7 @@ public: // Properties
 	float SunLightIntensity = 120000.0f;
 	
 	/** Base Moonlight Intensity (for Full Moon)
-	 * Typically 0.1 Lux, up to 0.32 Lux when the moon is at it's perigee (SuperMoon)
+	 * Typically 0.1 Lux, up to 0.32 Lux when the moon is at its perigee (SuperMoon)
 	 */
 	UPROPERTY(EditAnywhere, BlueprintReadWrite, AdvancedDisplay, Category="Celestial Vault", meta=( ForceUnits="Lux", ClampMin = "0", ClampMax = "200000"))
 	float MoonLightIntensity = 0.1f;
@@ -256,9 +301,13 @@ public:
 public: // Functions
 
 	/** Returns the current defined day, without any Time, because the Time will be controlled by the DaySequence Time of day - Uses "Now" or the Year/Month/Day properties */ 
-	UFUNCTION(BlueprintCallable, Category = "Celestial Vault|Date & Location")
+	UFUNCTION(BlueprintCallable, Category = "Celestial Vault|Date")
 	FDateTime GetDate() const;
 
+	/** Returns the current defined day, with the Time defined by the DaySequence Time of day*/ 
+	UFUNCTION(BlueprintCallable, Category = "Celestial Vault|Date")
+	FDateTime GetDateAndTime() const;
+	
 	/** Returns the Celestial Info for the Sun, at a specific Julian Date */  
 	UFUNCTION(BlueprintCallable, Category = "Celestial Vault|Sun")
 	FStellarBody GetSun(double JulianDate) const; 
@@ -285,26 +334,32 @@ public: // Functions
 	UFUNCTION(BlueprintCallable, Category = "Celestial Vault|Queries")
 	bool GetPlanetaryBodyByVSOP87Type(EVSOP87BodyType VSOP87Type, FPlanetaryBody& FoundPlanetaryBodyInfo, FTransform& BodyTransform);
 
-	/** Returns the Celestial vault rotation angle for the current Date at t=0 midnight */
-	double GetDayCelestialVaultAngle() const;
-
+	/** Returns the Vault Angle for a specific time - This is the GMST angle... */ 
+	UFUNCTION(BlueprintCallable, Category = "Celestial Vault|GMST")
+	double GetCelestialVaultAngle(double TimeOfDay) const;
+	
 #if WITH_EDITOR
 	
 	virtual void PostEditChangeProperty(FPropertyChangedEvent& PropertyChangedEvent) override;
+	virtual void PostEditChangeChainProperty(FPropertyChangedChainEvent& PropertyChangedEvent) override;
+	
 #endif
 
 private:
 	void InitStars();
 	void InitPlanetaryBodies();
 	FPlanetaryBody GetPlanetaryBodyInfo(const FPlanetaryBodyInputData& InputPlanetaryBody, double JulianDay) const;
+	void UpdatePlanetCenter();
 
 	UFUNCTION(BlueprintCallable, CallInEditor, Category= "Celestial Vault")
 	void RebuildAll();
 
 #if CELESTIAL_VAULT_ENABLE_DRAW_DEBUG
+
 	const FString CelestialVaultSF = TEXT("CelestialVault");
 	FDelegateHandle DrawDebugDelegateHandle ;
 	void DebugDrawCallback(UCanvas* Canvas, APlayerController* PC) const;
+
 #endif
 	
 };
diff --git a/Engine/Plugins/Experimental/CelestialVault/Source/CelestialVault/Public/DaylightSavings.h b/Engine/Plugins/Experimental/CelestialVault/Source/CelestialVault/Public/DaylightSavings.h
new file mode 100644
--- /dev/null
+++ b/Engine/Plugins/Experimental/CelestialVault/Source/CelestialVault/Public/DaylightSavings.h
@@ -0,0 +1,131 @@
+﻿// Copyright Epic Games, Inc. All Rights Reserved.
+
+
+#pragma once
+
+#include "CoreMinimal.h"
+#include "Kismet/BlueprintFunctionLibrary.h"
+#include "DaylightSavings.generated.h"
+
+
+UENUM(BlueprintType)
+enum class EDaylightSavingsMode : uint8
+{
+	None,       // Do not apply daylight saving time (DST).
+	Latitude,   // Determine hemisphere from the latitude sign (>= 0 = Northern, < 0 = Southern).
+	Northern,   // Force Northern Hemisphere rules (DST typically observed during Apr–Oct, i.e., around June).
+	Southern    // Force Southern Hemisphere rules (DST typically observed during Oct–Mar, i.e., around December).
+};
+
+/** Different Daylight Savings rules that apply across the globe*/
+UENUM(BlueprintType)
+enum class EDaylightSavingRuleKind : uint8
+{
+	NthDay,
+	LastDay,
+	FixedDate
+};
+
+/** Enumerates the days of a week in 7-days calendars. */
+UENUM(BlueprintType)
+enum class EWeekDay : uint8
+{
+	Monday = 0,
+	Tuesday = 1,
+	Wednesday = 2,
+	Thursday = 3,
+	Friday = 4,
+	Saturday = 5,
+	Sunday = 6
+};
+
+/** Enumerates the months of the year in 12-month calendars. */
+UENUM(BlueprintType)
+enum class EMonth: uint8
+{
+	None = 0 UMETA(Hidden),
+	January = 1,
+	February = 2,
+	March = 3,
+	April = 4,
+	May = 5,
+	June = 6,
+	July = 7,
+	August = 8,
+	September = 9,
+	October = 10,
+	November = 11,
+	December = 12
+};
+
+
+USTRUCT(BlueprintType)
+struct CELESTIALVAULT_API FDaylightSavingsRule
+{
+	GENERATED_BODY()
+
+	/** The rule defining the day to choose for the Daylight Saving Start/End  */
+	UPROPERTY(BlueprintReadWrite, EditAnywhere, Category= "Daylight Savings")
+	EDaylightSavingRuleKind Rule = EDaylightSavingRuleKind::NthDay;
+
+	/** Nth Weekday to consider when the rule is NthDay - eg 3 for the 3rd Sunday of January */
+	UPROPERTY(BlueprintReadWrite, EditAnywhere, Category= "Daylight Savings", meta = (ClampMin=1, ClampMax=5, EditCondition = "Rule == EDaylightSavingRuleKind::NthDay", EditConditionHides))
+	uint8 Nth = 1;
+
+	/** Reference Weekday for Nth and LastDay rules */
+	UPROPERTY(BlueprintReadWrite, EditAnywhere, Category= "Daylight Savings", meta = (EditCondition = "Rule == EDaylightSavingRuleKind::NthDay || Rule == EDaylightSavingRuleKind::LastDay", EditConditionHides))
+	EWeekDay WeekDay = EWeekDay::Sunday;
+
+	/** Number for the Day of the Month for Fixed Date rule (1..31) */
+	UPROPERTY(BlueprintReadWrite, EditAnywhere, Category= "Daylight Savings", meta = (ClampMin=1, ClampMax=31, EditCondition = "Rule == EDaylightSavingRuleKind::FixedDate", EditConditionHides))
+	uint8 DayOfMonth = 1;
+
+	/** Month the rule apply to */
+	UPROPERTY(BlueprintReadWrite, EditAnywhere, Category= "Daylight Savings")
+	EMonth Month = EMonth::January;
+
+#if WITH_EDITOR
+	void ClampToValidDate(int32 Year)
+	{
+		const int32 MaxDays = FDateTime::DaysInMonth(Year, FMath::Clamp(static_cast<int32>(Month), 1 , 12));
+		DayOfMonth = FMath::Clamp(DayOfMonth, 1, MaxDays);
+	}
+#endif
+};
+
+/** Utilities related to Daylight Savings */
+UCLASS()
+class CELESTIALVAULT_API UDaylightSavings : public UBlueprintFunctionLibrary
+{
+	GENERATED_BODY()
+
+public:
+	/** Return the Date defined by a specific Daylight Saving Rule */ 
+	UFUNCTION(BlueprintPure, Category="Daylight Savings")
+	static FDateTime ToDate(const FDaylightSavingsRule& DaylightSavingsRule, int Year);
+
+	/** Return the Date String defined by a specific Daylight Saving Rule */ 
+	UFUNCTION(BlueprintPure, Category="Daylight Savings")
+	static FString ToString(const FDaylightSavingsRule& DaylightSavingsRule);
+	
+	/** Return the Day for a specific Date */ 
+	UFUNCTION(BlueprintPure, Category="Daylight Savings")
+	static EWeekDay GetWeekDay(FDateTime Date);
+
+	/** Return the nth Weekday in a Month for a specific Year and Month - Returns (1..31), or 0 if the nth weekday does not exist */ 
+	UFUNCTION(BlueprintPure, Category="Daylight Savings")
+	static int NthWeekdayInMonth(int Year, EMonth Month, EWeekDay Weekday, int Nth);
+
+	/** Return the Date of the Last Weekday in a Month for a specific Year and Month (1..31) */
+	UFUNCTION(BlueprintPure, Category="Daylight Savings")
+	static int LastWeekdayInMonth(int Year, EMonth Month, EWeekDay Weekday);
+
+	/** Check if a specific date is inside a Daylight saving range defined by its rules. */
+	UFUNCTION(BlueprintPure, Category="Daylight Savings")
+	static bool IsDaylightSavings(FDateTime DateTime, EDaylightSavingsMode Mode, double Latitude, FDaylightSavingsRule StartDay, FDaylightSavingsRule EndDay, int SwitchHour);
+
+private:
+	static FString GetMonthString(EMonth Month);
+	static FString GetDayString(EWeekDay Day);
+	static FString GetOrdinalSuffix(int32 Number);
+};
diff --git a/Engine/Plugins/Experimental/CelestialVault/Source/CelestialVaultEditor/Private/CelestialVaultEditor.cpp b/Engine/Plugins/Experimental/CelestialVault/Source/CelestialVaultEditor/Private/CelestialVaultEditor.cpp
--- a/Engine/Plugins/Experimental/CelestialVault/Source/CelestialVaultEditor/Private/CelestialVaultEditor.cpp
+++ b/Engine/Plugins/Experimental/CelestialVault/Source/CelestialVaultEditor/Private/CelestialVaultEditor.cpp
@@ -2,16 +2,28 @@
 
 #include "CelestialVaultEditor.h"
 
+#include "DaylightSavingsCustomization.h"
+#include "DaylightSavings.h"
+#include "PropertyEditorModule.h"
+#include "Modules/ModuleManager.h"
+
 #define LOCTEXT_NAMESPACE "FCelestialVaultEditorModule"
 
 void FCelestialVaultEditorModule::StartupModule()
 {
-    
+	FPropertyEditorModule& PropertyModule = FModuleManager::LoadModuleChecked<FPropertyEditorModule>("PropertyEditor");
+
+	PropertyModule.RegisterCustomPropertyTypeLayout( FDaylightSavingsRule::StaticStruct()->GetFName(), FOnGetPropertyTypeCustomizationInstance::CreateStatic(&FDaylightSavingsCustomization::MakeInstance));
 }
 
 void FCelestialVaultEditorModule::ShutdownModule()
 {
-    
+	if (FModuleManager::Get().IsModuleLoaded("PropertyEditor"))
+	{
+		FPropertyEditorModule& PropertyModule = FModuleManager::GetModuleChecked<FPropertyEditorModule>("PropertyEditor");
+
+		PropertyModule.UnregisterCustomPropertyTypeLayout(FDaylightSavingsRule::StaticStruct()->GetFName());
+	}
 }
 
 #undef LOCTEXT_NAMESPACE
diff --git a/Engine/Plugins/Experimental/CelestialVault/Source/CelestialVaultEditor/Private/DaylightSavingsCustomization.cpp b/Engine/Plugins/Experimental/CelestialVault/Source/CelestialVaultEditor/Private/DaylightSavingsCustomization.cpp
new file mode 100644
--- /dev/null
+++ b/Engine/Plugins/Experimental/CelestialVault/Source/CelestialVaultEditor/Private/DaylightSavingsCustomization.cpp
@@ -0,0 +1,171 @@
+﻿// Copyright Epic Games, Inc. All Rights Reserved.
+
+
+
+#include "DaylightSavingsCustomization.h"
+#include "DetailLayoutBuilder.h"
+#include "DetailWidgetRow.h"
+#include "DaylightSavings.h"
+#include "IDetailChildrenBuilder.h"
+#include "IPropertyUtilities.h"
+#include "Widgets/Text/STextBlock.h"
+
+// Add default functionality here for any IFDSTDayCustomization functions that are not pure virtual.
+void FDaylightSavingsCustomization::CustomizeHeader(TSharedRef<IPropertyHandle> PropertyHandle, FDetailWidgetRow& HeaderRow, IPropertyTypeCustomizationUtils& CustomizationUtils)
+{
+	// Capture the handle by value so we can use it in the lambda
+	TSharedRef<IPropertyHandle> LocalHandle = PropertyHandle;
+
+	HeaderRow
+	.NameContent()
+	[
+		PropertyHandle->CreatePropertyNameWidget()
+	]
+	.ValueContent()
+	.MinDesiredWidth(200.f)
+	[
+		SNew(STextBlock)
+		.Text_Lambda([LocalHandle]()
+		{
+			void* RawData = nullptr;
+			if (LocalHandle->GetValueData(RawData) == FPropertyAccess::Success && RawData)
+			{
+				const FDaylightSavingsRule* DaylightSavingsRule = static_cast<const FDaylightSavingsRule*>(RawData);
+
+				FString Caption = UDaylightSavings::ToString(*DaylightSavingsRule);
+
+				// Specific to Celestial Vault
+				// Ideally, we would like to display the date, but we need a "Year" property in the parent object.
+				// And we also need to know the hemisphere, in case we are in the south hemisphere, and the End Date is on the next year!
+
+				TOptional<FNeededDaylightSavingsValues> CelestialVaultValues = TryGetNeededDaylightSavingsValuesFromOuter(LocalHandle);
+				if (CelestialVaultValues.IsSet())
+				{
+					int YearToUse = CelestialVaultValues.GetValue().Year;
+					
+					if (CelestialVaultValues.GetValue().Mode == EDaylightSavingsMode::Southern ||
+						(CelestialVaultValues.GetValue().Mode == EDaylightSavingsMode::Latitude && CelestialVaultValues.GetValue().Latitude < 0 ))
+					{
+						// We are in the south hemisphere. But we need to increase the year only for the DaylightSavingsEndDay Property! 
+						if (FProperty* DaylightSavingsRuleProp = LocalHandle->GetProperty())
+						{
+							const FName DaylightSavingsRulePropName = DaylightSavingsRuleProp->GetFName();
+							if (DaylightSavingsRulePropName == "DaylightSavingsEnd")
+							{
+								YearToUse++;
+							}
+						}
+					}
+
+					FDateTime Day = UDaylightSavings::ToDate(*DaylightSavingsRule, YearToUse);
+					Caption += Day.ToFormattedString(TEXT(" : (%d %b %Y)"));
+					
+				}
+				return FText::FromString(Caption);
+			}
+
+			// In case of multi selection
+			return FText();
+		})
+		.Font(IDetailLayoutBuilder::GetDetailFontItalic())
+	];
+}
+
+void FDaylightSavingsCustomization::CustomizeChildren(TSharedRef<IPropertyHandle> PropertyHandle, IDetailChildrenBuilder& ChildBuilder, IPropertyTypeCustomizationUtils& CustomizationUtils)
+{
+	// Default behavior: show the struct’s children as normal when expanded
+	uint32 NumChildren = 0;
+	PropertyHandle->GetNumChildren(NumChildren);
+
+	for (uint32 ChildIndex = 0; ChildIndex < NumChildren; ++ChildIndex)
+	{
+		ChildBuilder.AddProperty(PropertyHandle->GetChildHandle(ChildIndex).ToSharedRef());
+	}
+}
+
+TSharedRef<IPropertyTypeCustomization> FDaylightSavingsCustomization::MakeInstance()
+{
+	return MakeShareable(new FDaylightSavingsCustomization);
+}
+
+TOptional<FNeededDaylightSavingsValues> FDaylightSavingsCustomization::TryGetNeededDaylightSavingsValuesFromOuter(const TSharedRef<IPropertyHandle>& PropertyHandle)
+{
+	TArray<UObject*> Outers;
+	PropertyHandle->GetOuterObjects(Outers);
+
+	UObject* Outer = Outers.Num() > 0 ? Outers[0] : nullptr;
+	if (Outer)
+	{
+		FNeededDaylightSavingsValues Values;
+
+		if (ReadInt32Property(Outer, TEXT("Year"), Values.Year) &&
+			ReadDaylightSavingsModeProperty(Outer, TEXT("DaylightSavingsMode"), Values.Mode) &&
+			ReadDoubleProperty(Outer, TEXT("Latitude"), Values.Latitude))
+		{
+			return Values;
+		}    
+	}
+	return {};
+}
+
+bool FDaylightSavingsCustomization::ReadInt32Property(UObject* Object, FName PropertyName, int32& Out)
+{
+    if (Object)
+    {
+    	// Check Property Name
+    	if (FProperty* Property = Object->GetClass()->FindPropertyByName(PropertyName))
+    	{
+    		// Check Property Type
+    		if (const FIntProperty* IntProperty = CastField<FIntProperty>(Property))
+    		{
+    			Out = IntProperty->GetPropertyValue_InContainer(Object);
+    			return true;
+    		}
+    	}    
+    }
+    return false;
+}
+
+bool FDaylightSavingsCustomization::ReadDoubleProperty(UObject* Object, FName PropertyName, double& Out)
+{
+	if (Object)
+	{
+		// Check Property Name
+		if (FProperty* Property = Object->GetClass()->FindPropertyByName(PropertyName))
+		{
+			// Check Property Type
+			if (const FDoubleProperty* DoubleProperty = CastField<FDoubleProperty>(Property))
+			{
+				Out = DoubleProperty->GetPropertyValue_InContainer(Object);
+				return true;
+			}
+			if (const FFloatProperty* FloatProperty = CastField<FFloatProperty>(Property))
+			{
+				Out = static_cast<double>(FloatProperty->GetPropertyValue_InContainer(Object));
+				return true;
+			}
+		}
+	}
+    return false;
+}
+
+bool FDaylightSavingsCustomization::ReadDaylightSavingsModeProperty(UObject* Object, FName PropertyName, EDaylightSavingsMode& Out)
+{
+    if (Object) 
+    {
+    	// Check Property Name
+	    if (FProperty* Property = Object->GetClass()->FindPropertyByName(PropertyName))
+    	{
+	    	// Check Property Enum Type
+    		if (const FEnumProperty* EnumProperty = CastField<FEnumProperty>(Property))
+    		{
+    			const void* ValuePtr = EnumProperty->ContainerPtrToValuePtr<void>(Object);
+    			int64 Raw = EnumProperty->GetUnderlyingProperty()->GetSignedIntPropertyValue(ValuePtr);
+    			Out = static_cast<EDaylightSavingsMode>(Raw);
+    			return true;
+    		}
+    	}
+    }
+    return false;
+}
+
diff --git a/Engine/Plugins/Experimental/CelestialVault/Source/CelestialVaultEditor/Public/DaylightSavingsCustomization.h b/Engine/Plugins/Experimental/CelestialVault/Source/CelestialVaultEditor/Public/DaylightSavingsCustomization.h
new file mode 100644
--- /dev/null
+++ b/Engine/Plugins/Experimental/CelestialVault/Source/CelestialVaultEditor/Public/DaylightSavingsCustomization.h
@@ -0,0 +1,36 @@
+﻿// Copyright Epic Games, Inc. All Rights Reserved.
+
+#pragma once
+
+#include "CoreMinimal.h"
+#include "DaylightSavings.h"
+#include "IPropertyTypeCustomization.h"
+
+enum class EDaylightSavingsMode : uint8;
+class IPropertyHandle;
+
+struct FNeededDaylightSavingsValues
+{
+	int32 Year = 0;
+	EDaylightSavingsMode Mode = EDaylightSavingsMode::None;
+	double Latitude = 0.0;
+};
+
+
+class CELESTIALVAULTEDITOR_API FDaylightSavingsCustomization : public IPropertyTypeCustomization
+{
+public:
+	// Inherited from IPropertyTypeCustomization
+	virtual void CustomizeHeader(TSharedRef<IPropertyHandle> PropertyHandle, FDetailWidgetRow& HeaderRow, IPropertyTypeCustomizationUtils& CustomizationUtils) override;
+	virtual void CustomizeChildren(TSharedRef<IPropertyHandle> PropertyHandle, IDetailChildrenBuilder& ChildBuilder, IPropertyTypeCustomizationUtils& CustomizationUtils) override;
+	
+
+	// Utility function to create an instance of the Property Customization.
+	static TSharedRef<IPropertyTypeCustomization> MakeInstance();
+
+private:
+	static TOptional<FNeededDaylightSavingsValues> TryGetNeededDaylightSavingsValuesFromOuter(const TSharedRef<IPropertyHandle>& PropertyHandle);
+	static bool ReadInt32Property(UObject* Object, FName PropertyName, int32& Out);
+	static bool ReadDoubleProperty(UObject* Object, FName PropertyName, double& Out);
+	static bool ReadDaylightSavingsModeProperty(UObject* Object, FName PropertyName, EDaylightSavingsMode& Out);
+};
