* Now possible to run the Day Sequence over several days, with proper Moon movement around midnight.
 - Removed the Procedural Day Sequence that computed the bodies motions
 - Fixed the inconsistency of the default DSCA tracks
 - Put a local dynamic cache for celestial bodies' motion, generated on demand
 - Accelerated the Celestial computations by caching the Earth Kinematic State over several requests.
 - Consider the number of times the sequence has looped at runtime in every Date computations

* Improved the runtime performance by having the stars ISM inheriting from the Deepsky sphere bound instead of recomputing them each frame.
* Fog is no longer volumetric by default.
* Sun and Moon Info are now recomputed accurately for the current Time of day, rather than using cached information from midnight.
* Improved the values of the Daylight savings enums and their comments
* Random bug fixes and computational guardrails

#virtualized
#Synced-CL 49606359

#Synced-CL 49606359

diff --git a/Engine/Plugins/Experimental/CelestialVault/Source/CelestialVault/Private/CelestialMaths.cpp b/Engine/Plugins/Experimental/CelestialVault/Source/CelestialVault/Private/CelestialMaths.cpp
--- a/Engine/Plugins/Experimental/CelestialVault/Source/CelestialVault/Private/CelestialMaths.cpp
+++ b/Engine/Plugins/Experimental/CelestialVault/Source/CelestialVault/Private/CelestialMaths.cpp
@@ -67,6 +67,8 @@ double UCelestialMaths::GetIlluminationPercentage(double NormalizedAge)
 
 void UCelestialMaths::GetSolarSystemBodyLocationVelocity_VSOP87_AU(EVSOP87BodyType VSOP87BodyType, double JulianDate, FVector& LocationAU, FVector& Velocity)
 {
+	TRACE_CPUPROFILER_EVENT_SCOPE(GetSolarSystemBodyLocationVelocity_VSOP87_AU);
+	
 	// Warning - We need to use VSOP 87 Time for these functions! 
 	double VSOPTime = JulianDateToVSOP87Time(JulianDate);
 
@@ -112,6 +114,8 @@ void UCelestialMaths::GetSolarSystemBodyLocationVelocity_VSOP87_AU(EVSOP87BodyTy
 
 FVector UCelestialMaths::GetSolarSystemBodyLocation_VSOP87_Relativistic(FVector ReferenceBodyLocation_Heliocentric_AU, EVSOP87BodyType VSOP87BodyType, double JulianDate)
 {
+	TRACE_CPUPROFILER_EVENT_SCOPE(GetSolarSystemBodyLocation_VSOP87_Relativistic);
+	
 	double JulianDateLightAdjusted = JulianDate;
 	FVector BodyPositionAU = FVector::ZeroVector;
 
@@ -131,11 +135,38 @@ FVector UCelestialMaths::GetSolarSystemBodyLocation_VSOP87_Relativistic(FVector
 
 #pragma region Celestial Bodies
 
-void UCelestialMaths::GetBodyCelestialCoordinatesAU(double JulianDate, EVSOP87BodyType VSOP87BodyType, double ObserverLatitude, double ObserverLongitude,
-	double& RAJ2000Hours, double& DECJ2000Degrees, double& RAHours, double& DECDegrees, double& DistanceBodyToEarthAU, double& DistanceBodyToSunAU, double& DistanceEarthToSunAU) 
+FPlanetaryBodyKinematicState UCelestialMaths::GetPlanetaryBodyKinematicState_AU(double JulianDate, EVSOP87BodyType VSOP87BodyType)
+{
+	TRACE_CPUPROFILER_EVENT_SCOPE(GetPlanetaryBodyKinematicState_AU);
+	
+	FPlanetaryBodyKinematicState BodyKinematicState;
+
+	// Get Body Location in VSOP Rectangular Coordinates (X, Y, Z)
+	GetSolarSystemBodyLocationVelocity_VSOP87_AU(VSOP87BodyType, JulianDate, BodyKinematicState.Location_VSOP87, BodyKinematicState.Velocity_VSOP87);
+
+	// Transform to FK5J2000 Frame
+	BodyKinematicState.Location_Heliocentric_FK5J2000 = VSOPEclipticToFK5Equatorial.TransformVector(BodyKinematicState.Location_VSOP87);
+	BodyKinematicState.Velocity_FK5J2000 = VSOPEclipticToFK5Equatorial.TransformVector(BodyKinematicState.Velocity_VSOP87);
+
+	return BodyKinematicState;
+}
+
+void UCelestialMaths::GetBodyCelestialCoordinatesAU(double JulianDate, EVSOP87BodyType VSOP87BodyType, double ObserverLatitude, double ObserverLongitude, bool bIgnoreRelativisticEffect, 
+                                                    double& RAJ2000Hours, double& DECJ2000Degrees, double& RAHours, double& DECDegrees, double& DistanceBodyToEarthAU, double& DistanceBodyToSunAU, double& DistanceEarthToSunAU) 
 {
-	TRACE_CPUPROFILER_EVENT_SCOPE(UCelestialMaths::GetBodyCelestialCoordinatesAU);
+	TRACE_CPUPROFILER_EVENT_SCOPE(GetBodyCelestialCoordinatesAU);
+
+	// Compute Earth Reference location
+	FPlanetaryBodyKinematicState EarthKinematicState = GetPlanetaryBodyKinematicState_AU(JulianDate, EVSOP87BodyType::Earth);
 
+	// Call the actual computation method with the Earth as a reference. 
+	GetBodyCelestialCoordinatesAU_UsingKnownState(JulianDate, VSOP87BodyType, EarthKinematicState, ObserverLatitude, ObserverLongitude, bIgnoreRelativisticEffect, RAJ2000Hours, DECJ2000Degrees, RAHours, DECDegrees, DistanceBodyToEarthAU, DistanceBodyToSunAU, DistanceEarthToSunAU);
+}
+
+void UCelestialMaths::GetBodyCelestialCoordinatesAU_UsingKnownState(double JulianDate, EVSOP87BodyType VSOP87BodyType, const FPlanetaryBodyKinematicState& EarthKinematicState, double ObserverLatitude, double ObserverLongitude, bool bIgnoreRelativisticEffect, double& RAJ2000Hours, double& DECJ2000Degrees, double& RAHours, double& DECDegrees, double& DistanceBodyToEarthAU, double& DistanceBodyToSunAU, double& DistanceEarthToSunAU)
+{
+	TRACE_CPUPROFILER_EVENT_SCOPE(GetBodyCelestialCoordinatesAU_UsingKnownState);
+	
 	// Compute the Precession and Nutation Matrices - The formulas expect a Date in Terrestrial Time
 	double TAISeconds = JulianDateToInternationalAtomicTime(JulianDate);
 	double TTSeconds = InternationalAtomicTimeToTerrestrialTime(TAISeconds);
@@ -162,18 +193,19 @@ void UCelestialMaths::GetBodyCelestialCoordinatesAU(double JulianDate, EVSOP87Bo
 
 	// When getting the coordinates from any VSOP87 Call, it's important to change the axis frame to FK5J2000 using the VSOPEclipticToFK5Equatorial Matrix
 	// All the locations hereby defined will be using the FK5J2000 Frame, with HelioCentric or GeoCentric to precise the offset due to the Earth location
-	
-	// Get Earth Location in VSOP Rectangular Coordinates (X, Y, Z)
-	FVector EarthVelocity_VSOP87 = FVector::ZeroVector;
-	FVector EarthLocation_VSOP87 = FVector::ZeroVector;
-	GetSolarSystemBodyLocationVelocity_VSOP87_AU(EVSOP87BodyType::Earth, JulianDate, EarthLocation_VSOP87, EarthVelocity_VSOP87);
-	
-	// Transform to the FK5 J2000 Frame 
-	FVector EarthLocation_Heliocentric_FK5J2000 = VSOPEclipticToFK5Equatorial.TransformVector(EarthLocation_VSOP87); 
-	FVector EarthVelocity_FK5J2000 = VSOPEclipticToFK5Equatorial.TransformVector(EarthVelocity_VSOP87);
-	
+
 	// Get Body Location in FK5J2000 Rectangular Coordinates (X, Y, Z)
-	FVector BodyLocation_VSOP87 = GetSolarSystemBodyLocation_VSOP87_Relativistic(EarthLocation_VSOP87, VSOP87BodyType, JulianDate);
+	FVector BodyVelocity_VSOP87 = FVector::ZeroVector;
+	FVector BodyLocation_VSOP87 = FVector::ZeroVector;
+	if (bIgnoreRelativisticEffect)
+	{
+		GetSolarSystemBodyLocationVelocity_VSOP87_AU(VSOP87BodyType, JulianDate, BodyLocation_VSOP87, BodyVelocity_VSOP87);
+	}
+	else
+	{
+		BodyLocation_VSOP87 = GetSolarSystemBodyLocation_VSOP87_Relativistic(EarthKinematicState.Location_VSOP87, VSOP87BodyType, JulianDate);		
+	}
+	
 	FVector BodyLocation_Heliocentric_FK5J2000_AU = VSOPEclipticToFK5Equatorial.TransformVector(BodyLocation_VSOP87);
 
 	// Compute the Observer Location in the same, applying the effects of Precession matrix. 
@@ -200,14 +232,14 @@ void UCelestialMaths::GetBodyCelestialCoordinatesAU(double JulianDate, EVSOP87Bo
 
 	// Convert the XYZ Coordinates to RA, DEC - This will be the J2000 ones!
 	// Topocentric means “from the point of view of a specific location on Earth.” TopoCentric is not Geocentric! 
-	FVector BodyLocation_TopoCentric_FK5J2000_AU = BodyLocation_Heliocentric_FK5J2000_AU - EarthLocation_Heliocentric_FK5J2000 - ObserverLocationGeoCentric_FK5J2000AU;
+	FVector BodyLocation_TopoCentric_FK5J2000_AU = BodyLocation_Heliocentric_FK5J2000_AU - EarthKinematicState.Location_Heliocentric_FK5J2000 - ObserverLocationGeoCentric_FK5J2000AU;
 	double RAJ2000Degrees; 
 	XYZToRADEC_RH(BodyLocation_TopoCentric_FK5J2000_AU, RAJ2000Degrees, DECJ2000Degrees, DistanceBodyToEarthAU);
     RAJ2000Hours = RAJ2000Degrees / 15.0;
 
 	// Compute the RA and DEC of Date
 	// Transform the Body Topographic location with the Precession and Nutation Matrices at Date. 
-	FVector BodyLocation_TopoCentric_FK5J2000_Aberrated_AU = ComputeAberration(BodyLocation_TopoCentric_FK5J2000_AU, EarthVelocity_FK5J2000);
+	FVector BodyLocation_TopoCentric_FK5J2000_Aberrated_AU = ComputeAberration(BodyLocation_TopoCentric_FK5J2000_AU, EarthKinematicState.Velocity_FK5J2000);
 	FVector BodyLocation_TopoCentric_FK5J2000_Precessed_AU = PrecessionMatrix.InverseTransformVector(BodyLocation_TopoCentric_FK5J2000_Aberrated_AU);
 	FVector BodyLocation_TopoCentric_FK5J2000_OfDate_AU = NutationMatrix.InverseTransformVector(BodyLocation_TopoCentric_FK5J2000_Precessed_AU);
 	
@@ -215,15 +247,25 @@ void UCelestialMaths::GetBodyCelestialCoordinatesAU(double JulianDate, EVSOP87Bo
 	XYZToRADEC_RH(BodyLocation_TopoCentric_FK5J2000_OfDate_AU, RADegrees, DECDegrees, DistanceBodyToEarthAU);
 	RAHours = RADegrees / 15.0;
 
-	DistanceBodyToEarthAU = (BodyLocation_Heliocentric_FK5J2000_AU - EarthLocation_Heliocentric_FK5J2000).Length();
+	DistanceBodyToEarthAU = (BodyLocation_Heliocentric_FK5J2000_AU - EarthKinematicState.Location_Heliocentric_FK5J2000).Length();
 	DistanceBodyToSunAU = BodyLocation_Heliocentric_FK5J2000_AU.Length();
-	DistanceEarthToSunAU = EarthLocation_Heliocentric_FK5J2000.Length();
+	DistanceEarthToSunAU = EarthKinematicState.Location_Heliocentric_FK5J2000.Length();
 }
 
+
+
 double UCelestialMaths::GetPlanetaryBodyMagnitude(EVSOP87BodyType VSOP87BodyType, double DistanceToSunAU, double DistanceToEarthAU, double DistanceEarthToSunAU, double& PhaseAngle)
 {
+	TRACE_CPUPROFILER_EVENT_SCOPE(GetPlanetaryBodyMagnitude);
+	
 	// From Computing Apparent Planetary Magnitudes for The Astronomical Almanac
 	// James L. Hilton US Naval Observatory
+
+	if (VSOP87BodyType == EVSOP87BodyType::Sun)
+	{
+		// Early out for the sun, and also protects again a divide by zero with DistanceToSunAU in the PhaseAngle Computation
+		return -26.74;
+	}
 	
 	// V = 5 log10 ( r d ) + V1(0) + C1 α + C2 α2 + ... with
 	//   r = planet’s distance from the Sun
@@ -232,6 +274,8 @@ double UCelestialMaths::GetPlanetaryBodyMagnitude(EVSOP87BodyType VSOP87BodyType
 	//   V1(0) sometimes referred to as the planet’s absolute magnitude or geometric magnitude is the magnitude when observed at α = 0
 	//   ΣnCn αn is called the phase function 
 	
+
+	
 	
 	double DistanceFactor = 5.0 * FMath::LogX(10.0, DistanceToEarthAU * DistanceToSunAU);
 	double ApparentMagnitude = 0.0; 
@@ -358,6 +402,8 @@ FDateTime UCelestialMaths::UTCTimeToLocalTime(FDateTime UTCTime, double TimeZone
 
 double UCelestialMaths::UTCDateTimeToJulianDate(FDateTime UTCDateTime)
 {
+	TRACE_CPUPROFILER_EVENT_SCOPE(UTCDateTimeToJulianDate);
+	
 	// From https://www.celestialprogramming.com/julian.html
 
 	// Get Individual values for YMD and HMS
@@ -938,6 +984,17 @@ FVector UCelestialMaths::ComputeAberration(const FVector& TopocentricTargetLocat
 
 FStellarBody UCelestialMaths::GetSunInformation(double JulianDate, double ObserverLatitude, double ObserverLongitude)
 {
+	TRACE_CPUPROFILER_EVENT_SCOPE(GetSunInformation);
+	
+	// Compute Earth Reference location
+	FPlanetaryBodyKinematicState EarthKinematicState = GetPlanetaryBodyKinematicState_AU(JulianDate, EVSOP87BodyType::Earth);
+
+	return GetSunInformation_UsingKnownState(JulianDate, EarthKinematicState, ObserverLatitude, ObserverLongitude);
+}
+
+FStellarBody UCelestialMaths::GetSunInformation_UsingKnownState(double JulianDate, const FPlanetaryBodyKinematicState& EarthKinematicState, double ObserverLatitude, double ObserverLongitude)
+{
+	TRACE_CPUPROFILER_EVENT_SCOPE(GetSunInformation_UsingKnownState);
 	// Get the Celestial Body Coordinates
 	double RAJ2000Hours;
 	double DECJ2000Degrees;
@@ -946,7 +1003,7 @@ FStellarBody UCelestialMaths::GetSunInformation(double JulianDate, double Observ
 	double DistanceBodyToEarthAU;
 	double DistanceBodyToSunAU;
 	double DistanceEarthToSunAU;
-	GetBodyCelestialCoordinatesAU(JulianDate, EVSOP87BodyType::Sun, ObserverLatitude, ObserverLongitude, RAJ2000Hours, DECJ2000Degrees, RAHours, DECDegrees, DistanceBodyToEarthAU, DistanceBodyToSunAU, DistanceEarthToSunAU);  
+	GetBodyCelestialCoordinatesAU_UsingKnownState(JulianDate, EVSOP87BodyType::Sun, EarthKinematicState, ObserverLatitude, ObserverLongitude, false, RAJ2000Hours, DECJ2000Degrees, RAHours, DECDegrees, DistanceBodyToEarthAU, DistanceBodyToSunAU, DistanceEarthToSunAU);  
 
 	// Set the Sun properties. Some properties can be hardcoded for the Sun (CI, Magnitude, Radius) 
 	FStellarBody SunInfo;
diff --git a/Engine/Plugins/Experimental/CelestialVault/Source/CelestialVault/Private/CelestialVaultDaySequenceActor.cpp b/Engine/Plugins/Experimental/CelestialVault/Source/CelestialVault/Private/CelestialVaultDaySequenceActor.cpp
--- a/Engine/Plugins/Experimental/CelestialVault/Source/CelestialVault/Private/CelestialVaultDaySequenceActor.cpp
+++ b/Engine/Plugins/Experimental/CelestialVault/Source/CelestialVault/Private/CelestialVaultDaySequenceActor.cpp
@@ -51,7 +51,6 @@ ACelestialVaultDaySequenceActor::ACelestialVaultDaySequenceActor(const FObjectIn
 {
 	ExponentialHeightFogComponent = CreateOptionalDefaultSubobject<UExponentialHeightFogComponent>(TEXT("ExponentialHeightFog"));
 	ExponentialHeightFogComponent->SetupAttachment(RootComponent);
-	ExponentialHeightFogComponent->bEnableVolumetricFog = true;
 
 	static ConstructorHelpers::FObjectFinder<UCurveFloat> HighlightContrastCurve(TEXT("/CelestialVault/Data/CF_CelestialHighlightContrastCurve.CF_CelestialHighlightContrastCurve"));
 	GlobalPostProcessVolume = CreateOptionalDefaultSubobject<UPostProcessComponent>(TEXT("GlobalPostProcessVolume"));
@@ -132,7 +131,7 @@ ACelestialVaultDaySequenceActor::ACelestialVaultDaySequenceActor(const FObjectIn
 	static ConstructorHelpers::FObjectFinder<UStaticMesh> PlaneXMesh(TEXT("/CelestialVault/Meshes/SM_Plane_FacingX.SM_Plane_FacingX"));
 	static ConstructorHelpers::FObjectFinder<UMaterialInterface> StarsDefaultMaterial(TEXT("/CelestialVault/Materials/MI_Stars.MI_Stars"));
 	StarsComponent = CreateOptionalDefaultSubobject<UInstancedStaticMeshComponent>(TEXT("Stars"));
-	StarsComponent->SetupAttachment(CelestialVaultComponent);
+	StarsComponent->SetupAttachment(DeepSkyComponent);
 	StarsComponent->SetStaticMesh(PlaneXMesh.Object);
 	StarsComponent->SetCollisionEnabled(ECollisionEnabled::NoCollision);
 	StarsComponent->SetGenerateOverlapEvents(false);
@@ -140,6 +139,7 @@ ACelestialVaultDaySequenceActor::ACelestialVaultDaySequenceActor(const FObjectIn
 	StarsComponent->SetAffectDynamicIndirectLighting(false);
 	StarsComponent->SetCanEverAffectNavigation(false);
 	StarsComponent->SetMaterial(0, StarsDefaultMaterial.Object);
+	StarsComponent->bUseAttachParentBound = true;
 
 	// Planets ISM
 	static ConstructorHelpers::FObjectFinder<UMaterialInterface> PlanetsDefaultMaterial(TEXT("/CelestialVault/Materials/MI_SolarSystemPlanets.MI_SolarSystemPlanets"));
@@ -152,6 +152,7 @@ ACelestialVaultDaySequenceActor::ACelestialVaultDaySequenceActor(const FObjectIn
 	PlanetsComponent->SetAffectDynamicIndirectLighting(false);
 	PlanetsComponent->SetCanEverAffectNavigation(false);
 	PlanetsComponent->SetMaterial(0, PlanetsDefaultMaterial.Object);
+	
 
 	// Moon - Disc
 	static ConstructorHelpers::FObjectFinder<UMaterialInterface> MoonDiscDefaultMaterial(TEXT("/CelestialVault/Materials/MI_Moon.MI_Moon"));
@@ -207,9 +208,8 @@ bool ACelestialVaultDaySequenceActor::IsDaylightSavingsNow() const
 
 bool ACelestialVaultDaySequenceActor::IsDaylightSavingsAtTimeOfDay(double TimeOfDay) const
 {
-	FDateTime Date = GetDate();
-	FTimespan ElapsedTimeOfDay = FTimespan::FromHours(TimeOfDay);
-	return UDaylightSavings::IsDaylightSavings(Date + ElapsedTimeOfDay, DaylightSavingsMode, Latitude, DaylightSavingsStart, DaylightSavingsEnd, DaylightSavingsSwitchHour);
+	FDateTime LocalTime = GetDate() + FTimespan::FromHours(TimeOfDay);
+	return UDaylightSavings::IsDaylightSavings(LocalTime,  DaylightSavingsMode, Latitude, DaylightSavingsStart, DaylightSavingsEnd, DaylightSavingsSwitchHour);
 }
 
 //////// Protected Functions
@@ -225,8 +225,7 @@ void ACelestialVaultDaySequenceActor::BeginPlay()
 		{
 			DaySequenceSubsystem->SetDaySequenceActor(this);
 		}
-
-		//RebuildAll(); // TODO - is it needed? 
+		SequenceLoopCount = 0;
 	}
 }
 
@@ -260,6 +259,45 @@ void ACelestialVaultDaySequenceActor::OnConstruction(const FTransform& Transform
 	}
 }
 
+void ACelestialVaultDaySequenceActor::Tick(float DeltaTime)
+{
+	Super::Tick(DeltaTime);
+
+#if WITH_EDITOR
+	// In Editor, DaySequenceUpdated is not called, but the actor is ticked
+	// In this case, we'll update the Bodies motion manually on tick to react to the TimeOfDay Slider changes
+	if (UWorld* World = GetWorld())
+	{
+		// We don't want to update the bodies motion each frame in Game, as the DaySequence actor has a Sequence Update Interval to allow for sliced updates. We don't want to override this feature.   
+		if (!World->IsGameWorld())
+		{
+			// In editor...
+			UpdateBodiesMotion();	
+		}
+	}
+#endif
+}
+
+void ACelestialVaultDaySequenceActor::SequencePlayerUpdated(float CurrentTime, float PreviousTime)
+{
+	Super::SequencePlayerUpdated(CurrentTime, PreviousTime);
+
+	if (IDaySequencePlayer* Player = GetSequencePlayer())
+	{
+		if (Player->GetCurrentNumLoops() != SequenceLoopCount)
+		{
+			SequenceLoopCount = Player->GetCurrentNumLoops(); 
+			UE_LOG(LogCelestialVault, Verbose, TEXT("Day Sequence Loop: We are now on Day %d"), SequenceLoopCount);
+
+			// New day, refresh the Body Motion Cache.  
+			ResetBodyMotionCache();
+		}
+
+		// Update Celestial Bodies Motion.
+		UpdateBodiesMotion();
+	}
+}
+
 void ACelestialVaultDaySequenceActor::PostRegisterAllComponents()
 {
 	Super::PostRegisterAllComponents();
@@ -314,19 +352,19 @@ void ACelestialVaultDaySequenceActor::DebugDrawCallback(UCanvas* Canvas, APlayer
 		DebugMessage += FString::Printf(TEXT("    Julian Centuries: %.5f"), JulianCenturies) + LINE_TERMINATOR;
 		DebugMessage += FString::Printf(TEXT("    Daylight Savings: %s"), IsDaylightSavingsNow() ? TEXT("Yes") : TEXT("No") ) + LINE_TERMINATOR;
 		
-	DebugMessage += FString::Printf(TEXT("    GMST at 0:00: %s (%.5f)°"), *UCelestialMaths::Conv_DegreesToHMSString(UCelestialMaths::ModPositive(GMST0_Unwrapped,360)), UCelestialMaths::ModPositive(GMST0_Unwrapped,360)) + LINE_TERMINATOR;
+		DebugMessage += FString::Printf(TEXT("    GMST at 0:00: %s (%.5f)°"), *UCelestialMaths::Conv_DegreesToHMSString(UCelestialMaths::ModPositive(GMST0_Unwrapped,360)), UCelestialMaths::ModPositive(GMST0_Unwrapped,360)) + LINE_TERMINATOR;
 		DebugMessage += FString::Printf(TEXT("    GMST: %s (%.5f)°"), *UCelestialMaths::Conv_DegreesToHMSString(GMST), GMST) + LINE_TERMINATOR;
 		DebugMessage += FString::Printf(TEXT("    LMST: %s (%.5f)°"), *UCelestialMaths::Conv_DegreesToHMSString(LMST), LMST) + LINE_TERMINATOR;
 
 		// Sun and Moon Info
 		DebugMessage += LINE_TERMINATOR;
 		DebugMessage += FString("    -- Sun -- \n");
-		FStellarBody SunAtTime = GetSun(JulianDate);
+		FStellarBody SunAtTime = ComputeSunInfo(JulianDate);
 		DebugMessage += SunAtTime.ToString() + LINE_TERMINATOR;
 
 		DebugMessage += FString("    -- Moon -- \n");
-		FPlanetaryBody MoonAtTime = GetMoon(JulianDate);
-		DebugMessage += Moon.ToString() + LINE_TERMINATOR;
+		FPlanetaryBody MoonAtTime = ComputeMoonInfo(JulianDate);
+		DebugMessage += MoonAtTime.ToString() + LINE_TERMINATOR;
 		
 		// Draw
 		Canvas->DrawText(GetDefault<UEngine>()->GetMediumFont(), DebugMessage, 10, 10);
@@ -351,19 +389,40 @@ FDateTime ACelestialVaultDaySequenceActor::GetDate() const
 		return FDateTime(Year, ClampedMonth, ClampedDay, 0,0,0);
 	}
 		
-	return FDateTime(Year, static_cast<int>(Month), Day, 0,0,0);
+	return FDateTime(Year, static_cast<int>(Month), Day, 0,0,0) + FTimespan(SequenceLoopCount, 0,0, 0);;
 }
 
 FDateTime ACelestialVaultDaySequenceActor::GetDateAndTime() const
 {
-	FDateTime Date = GetDate();
-	FTimespan ElapsedTimeOfDay = FTimespan::FromHours(GetTimeOfDay());
-	return Date + ElapsedTimeOfDay;
+	return GetDate() + FTimespan::FromHours(GetTimeOfDay());
+}
+
+double ACelestialVaultDaySequenceActor::GetJulianDate() const
+{
+	FDateTime LocalTime = GetDateAndTime();
+	FDateTime UTCTime = UCelestialMaths::LocalTimeToUTCTime(LocalTime, GMT_TimeZone, IsDaylightSavingsNow());
+	return UCelestialMaths::UTCDateTimeToJulianDate(UTCTime);
+}
+
+FStellarBody ACelestialVaultDaySequenceActor::ComputeSunInfo(double JulianDate) const
+{
+	// Compute Earth Reference location
+	FPlanetaryBodyKinematicState EarthKinematicState = UCelestialMaths::GetPlanetaryBodyKinematicState_AU(JulianDate, EVSOP87BodyType::Earth);
+
+	return ComputeSunInfo_UsingKnownState(JulianDate, EarthKinematicState);
+}
+
+FStellarBody ACelestialVaultDaySequenceActor::ComputeCurrentSunInfo() const
+{
+	return ComputeSunInfo(GetJulianDate());
 }
 
-FStellarBody ACelestialVaultDaySequenceActor::GetSun(double JulianDate) const
+FStellarBody ACelestialVaultDaySequenceActor::ComputeSunInfo_UsingKnownState(double JulianDate, const FPlanetaryBodyKinematicState& EarthKinematicState) const
 {
-	FStellarBody ResultSunInfo = UCelestialMaths::GetSunInformation(JulianDate, Latitude, Longitude);
+	TRACE_CPUPROFILER_EVENT_SCOPE(ComputeSunInfo_UsingKnownState);
+
+	
+	FStellarBody ResultSunInfo = UCelestialMaths::GetSunInformation_UsingKnownState(JulianDate, EarthKinematicState, Latitude, Longitude);
 
 	// Location
 	FVector SunLocation = UCelestialMaths::RADECToXYZ_RH(ResultSunInfo.RA * 15.0, ResultSunInfo.DEC, UCelestialMaths::AstronomicalUnitsToMeters(ResultSunInfo.DistanceInAU) * 100); 
@@ -373,15 +432,31 @@ FStellarBody ACelestialVaultDaySequenceActor::GetSun(double JulianDate) const
 	return ResultSunInfo;
 }
 
-FPlanetaryBody ACelestialVaultDaySequenceActor::GetMoon(double JulianDate) const
+FPlanetaryBody ACelestialVaultDaySequenceActor::ComputeMoonInfo(double JulianDate) const
+{
+	// Compute Earth Reference location
+	FPlanetaryBodyKinematicState EarthKinematicState = UCelestialMaths::GetPlanetaryBodyKinematicState_AU(JulianDate, EVSOP87BodyType::Earth);
+	
+	return ComputeMoonInfo_UsingKnownState(JulianDate, EarthKinematicState);
+}
+
+FPlanetaryBody ACelestialVaultDaySequenceActor::ComputeCurrentMoonInfo() const
+{
+	return ComputeMoonInfo(GetJulianDate());
+}
+
+FPlanetaryBody ACelestialVaultDaySequenceActor::ComputeMoonInfo_UsingKnownState(double JulianDate, const FPlanetaryBodyKinematicState& EarthKinematicState) const
 {
-	double UEDistance = CelestialVaultDistance * UE_ONE_KILOMETER * MoonVaultPercentage / 100.0; 
-	FPlanetaryBody ResultMoonInfo = GetPlanetaryBodyInfo(FPlanetaryBodyInputData::Moon, JulianDate);
+	TRACE_CPUPROFILER_EVENT_SCOPE(ComputeMoonInfo_UsingKnownState);
+
+	
+	double UEDistance = CelestialVaultDistance * UE_ONE_KILOMETER * MoonVaultPercentage / 100.0;
+	FPlanetaryBody ResultMoonInfo = GetPlanetaryBodyInfo_UsingKnownState(JulianDate, FPlanetaryBodyInputData::Moon,  EarthKinematicState, true); // Ignore the relativistic effects for the moon. It's close to us, and save some iterations in the computations. 
 
 	// Set the Moon specifics, and compute the Transform
 	ResultMoonInfo.Age = UCelestialMaths::GetMoonNormalizedAgeSimple(JulianDate);;
 	ResultMoonInfo.IlluminationPercentage = UCelestialMaths::GetIlluminationPercentage(ResultMoonInfo.Age);
-	ResultMoonInfo.ComputeUELocalTransform();
+	ResultMoonInfo.ComputeUELocalTransform(MoonScale, UEDistance);
 	
 	if (bManualControl) // Override 
 	{
@@ -390,7 +465,7 @@ FPlanetaryBody ACelestialVaultDaySequenceActor::GetMoon(double JulianDate) const
 		ResultMoonInfo.IlluminationPercentage = UCelestialMaths::GetIlluminationPercentage(MoonAge);
 
 		// Location
-		FStellarBody SunInfoTemp = GetSun(JulianDate);
+		FStellarBody SunInfoTemp = ComputeSunInfo_UsingKnownState(JulianDate,  EarthKinematicState);
 		ResultMoonInfo.RA = UCelestialMaths::ModPositive(SunInfoTemp.RA + MoonOffset_RA, 24.0);
 		ResultMoonInfo.DEC = UCelestialMaths::ModPositive(SunInfoTemp.DEC + MoonOffset_DEC + 180.0, 360.0) - 180.0;
 		// We changed the RA, so we need to update the transform
@@ -399,6 +474,7 @@ FPlanetaryBody ACelestialVaultDaySequenceActor::GetMoon(double JulianDate) const
 	return ResultMoonInfo;
 }
 
+
 void ACelestialVaultDaySequenceActor::SetMoonDiscAge(float InMoonAge)
 {
 	if (MoonDiscComponent)
@@ -438,6 +514,7 @@ bool ACelestialVaultDaySequenceActor::GetClosestStar(FVector ObserverLocation, F
 		FTransform WorldTransform =  StarInfo.UELocalTransform * StarsComponentTransform;
 			
 		FVector DirectionToInstance = WorldTransform.GetLocation() - ObserverLocation;
+		DirectionToInstance.Normalize();
 
 		double CosDeltaAngle = FVector::DotProduct(LookupDirection, DirectionToInstance);
 		if (CosDeltaAngle > CosThresholdAngle && CosDeltaAngle > MinCos)
@@ -457,12 +534,6 @@ bool ACelestialVaultDaySequenceActor::GetClosestPlanetaryBody(FVector ObserverLo
 {
 	BodyTransform = FTransform::Identity;
 	
-	// We query only if we generated the Planet Info  
-	if (!bKeepPlanetsInfos)
-	{
-		return false;
-	}
-	
 	double CosThresholdAngle = FMath::Cos(FMath::DegreesToRadians(ThresholdAngleDegrees));
 	LookupDirection.Normalize();
 	double MinCos = -1.0;
@@ -482,7 +553,12 @@ bool ACelestialVaultDaySequenceActor::GetClosestPlanetaryBody(FVector ObserverLo
 			// Inside the cone angle. At least we have found the moon. 
 			MinCos = CosDeltaAngle;
 			BodyTransform = MoonDiscComponent->GetComponentTransform();
-			FoundPlanetaryBodyInfo = Moon;
+
+			// Compute the moon info for the exact moment in time
+			FDateTime LocalTime = GetDateAndTime();
+			FDateTime UTCTime = UCelestialMaths::LocalTimeToUTCTime(LocalTime, GMT_TimeZone, IsDaylightSavingsNow());
+			double JulianDate = UCelestialMaths::UTCDateTimeToJulianDate(UTCTime);
+			FoundPlanetaryBodyInfo = ComputeMoonInfo(JulianDate);
 			BodyFound = true;
 		}
 	}
@@ -490,25 +566,34 @@ bool ACelestialVaultDaySequenceActor::GetClosestPlanetaryBody(FVector ObserverLo
 	// Check for planets now
 	if (PlanetsComponent)
 	{
-		FTransform PlanetComponentTransform = PlanetsComponent->GetComponentTransform();
-		for (FPlanetaryBody BodyInfo : Planets)
+		// We query only if we generated the Planet Info  
+		if (bKeepPlanetsInfos)
 		{
-			// We need to compute the world Transform for the body, because the celestial vault has rotated and the Body UE Transform is local to the Sky. 
-			FTransform WorldTransform = BodyInfo.UELocalTransform * PlanetComponentTransform;
+			FTransform PlanetComponentTransform = PlanetsComponent->GetComponentTransform();
+			for (FPlanetaryBody BodyInfo : Planets)
+			{
+				// We need to compute the world Transform for the body, because the celestial vault has rotated and the Body UE Transform is local to the Sky. 
+				FTransform WorldTransform = BodyInfo.UELocalTransform * PlanetComponentTransform;
 
-			FVector DirectionToInstance = WorldTransform.GetLocation() - ObserverLocation;
-			DirectionToInstance.Normalize();
+				FVector DirectionToInstance = WorldTransform.GetLocation() - ObserverLocation;
+				DirectionToInstance.Normalize();
 
-			double CosDeltaAngle = FVector::DotProduct(LookupDirection, DirectionToInstance);
-			if (CosDeltaAngle > CosThresholdAngle && CosDeltaAngle > MinCos)
-			{
-				// Inside the cone angle, and closer to latest found body
-				MinCos = CosDeltaAngle;
-				FoundPlanetaryBodyInfo = BodyInfo;
-				BodyTransform = WorldTransform;
-				BodyFound = true;
+				double CosDeltaAngle = FVector::DotProduct(LookupDirection, DirectionToInstance);
+				if (CosDeltaAngle > CosThresholdAngle && CosDeltaAngle > MinCos)
+				{
+					// Inside the cone angle, and closer to latest found body
+					MinCos = CosDeltaAngle;
+					FoundPlanetaryBodyInfo = BodyInfo;
+					BodyTransform = WorldTransform;
+					BodyFound = true;
+				}
 			}
 		}
+		else
+		{
+			
+		}
+		
 	}
 
 	return BodyFound;
@@ -518,28 +603,32 @@ bool ACelestialVaultDaySequenceActor::GetPlanetaryBodyByVSOP87Type(EVSOP87BodyTy
 {
 	BodyTransform = FTransform::Identity;
 	
-	// We query only if we generated the Planets Info  
-	if (!bKeepPlanetsInfos)
-	{
-		return false;
-	}
-	
 	if (VSOP87Type == EVSOP87BodyType::Moon && MoonDiscComponent) 
 	{
 		BodyTransform = MoonDiscComponent->GetComponentTransform();
-		FoundPlanetaryBodyInfo = Moon;
+
+		// Compute the moon info for the exact moment in time
+		FDateTime LocalTime = GetDateAndTime();
+		FDateTime UTCTime = UCelestialMaths::LocalTimeToUTCTime(LocalTime, GMT_TimeZone, IsDaylightSavingsNow());
+		double JulianDate = UCelestialMaths::UTCDateTimeToJulianDate(UTCTime);
+		FoundPlanetaryBodyInfo = ComputeMoonInfo(JulianDate);
+		
 		return true;
 	}
 	else
 	{
-		for (FPlanetaryBody PlanetaryBodyInfo : Planets)
+		// We query only if we generated the Planets Info  
+		if (bKeepPlanetsInfos)
 		{
-			if (PlanetaryBodyInfo.VSOP87BodyType == VSOP87Type)
+			for (FPlanetaryBody PlanetaryBodyInfo : Planets)
 			{
-				FoundPlanetaryBodyInfo = PlanetaryBodyInfo;
-				return true;
+				if (PlanetaryBodyInfo.VSOP87BodyType == VSOP87Type)
+				{
+					FoundPlanetaryBodyInfo = PlanetaryBodyInfo;
+					return true;
+				}
 			}
-		}	
+		}
 	}
 	return false;
 }
@@ -575,7 +664,8 @@ void ACelestialVaultDaySequenceActor::PostEditChangeProperty(FPropertyChangedEve
 	{
 		if (DeepSkyComponent)
 		{
-			DeepSkyComponent->SetRelativeScale3D(FVector(CelestialVaultDistance*1000.0));	
+			// Sphere mesh 1 meter so 100 units, so to have km, we just have to multiply by 1000. 
+			DeepSkyComponent->SetRelativeScale3D(FVector(CelestialVaultDistance * 1000.0 ));	
 		}
 		
 		// TODO UpdateGeoReferencing();
@@ -644,24 +734,22 @@ void ACelestialVaultDaySequenceActor::PostEditChangeProperty(FPropertyChangedEve
 	
 	if (bBuildAll || bBuildStars)
 	{
-		InitStars();
+		InitStarsComponent();
 	}
 
 	if (bBuildAll || bBuildPlanetaryBodies)
 	{
-		InitPlanetaryBodies();
+		InitPlanetsComponent();
 	}
 
 	if (bRebuildSequence)
 	{
-		RootSequence = nullptr;
-		UpdateRootSequence();
+		ResetBodyMotionCache();
 	}
 
 	Super::PostEditChangeProperty(PropertyChangedEvent);
 }
 
-
 void ACelestialVaultDaySequenceActor::PostEditChangeChainProperty(FPropertyChangedChainEvent& PropertyChangedEvent)
 {
 	Super::PostEditChangeChainProperty(PropertyChangedEvent);
@@ -693,11 +781,11 @@ void ACelestialVaultDaySequenceActor::PostEditChangeChainProperty(FPropertyChang
 }
 #endif
 
-//////// Private  Functions
+//////// Private Functions
 
-void ACelestialVaultDaySequenceActor::InitStars() 
+void ACelestialVaultDaySequenceActor::InitStarsComponent() 
 {
-	TRACE_CPUPROFILER_EVENT_SCOPE(ACelestialVaultDaySequenceActor::InitStars);
+	TRACE_CPUPROFILER_EVENT_SCOPE(InitStarsComponent);
 	
 	if (StarsComponent)
 	{
@@ -718,17 +806,19 @@ void ACelestialVaultDaySequenceActor::InitStars()
 			const UScriptStruct* CelestialStarCatalogRowStruct = CelestialStarCatalog->GetRowStruct();
 			if (CelestialStarCatalogRowStruct != FCelestialStarInputData::StaticStruct())
 			{
-				UE_LOG(LogCelestialVault, Warning, TEXT("Invalid DataTable row structure for CelestialStarCatalog! It should be of type %s"), *FCelestialStarInputData::StaticStruct()->GetName() );
+				UE_LOG(LogCelestialVault, Warning, TEXT("Invalid DataTable row structure for CelestialStarCatalog! It should be of type %s"), *FCelestialStarInputData::StaticStruct()->GetName());
 			}
 			else
 			{
 				// Generate Celestial Stars
-				CelestialStarCatalog->ForeachRow<FCelestialStarInputData>("AEarthDaySequenceActor::InitStars", [&](const FName& Key, const FCelestialStarInputData& CelestialStarInputData)
+				CelestialStarCatalog->ForeachRow<FCelestialStarInputData>("AEarthDaySequenceActor::InitStarsComponent", [&](const FName& Key, const FCelestialStarInputData& CelestialStarInputData)
 					{
 						if (CelestialStarInputData.Magnitude <= MaxVisibleMagnitude)
 						{
 							// Location - Convert to UE Left-handed Frame (Invert the Y coordinate)
-							FVector StarLocation = UCelestialMaths::RADECToXYZ_RH(CelestialStarInputData.RA * 15.0, CelestialStarInputData.DEC, CelestialVaultDistance * 1000.0 * StarsVaultPercentage / 100.0) *  FVector(100.0, -100.0, 100.0);
+							// Special case for the Stars. To avoid recomputing the ISM bounds, we use the parent bounds, and therefore inherit from the transformation of the Celestial Vault.
+							// The mesh being 100 units radius, we have a local transform corresponding to a StarsVaultPercentage distance. 
+							FVector StarLocation = UCelestialMaths::RADECToXYZ_RH(CelestialStarInputData.RA * 15.0, CelestialStarInputData.DEC, StarsVaultPercentage) * FVector(1.0, -1.0, 1.0);
 
 							// Color
 							FLinearColor StarColor = UCelestialMaths::BVtoLinearColor(CelestialStarInputData.ColorIndex);
@@ -757,8 +847,8 @@ void ACelestialVaultDaySequenceActor::InitStars()
 								StarInfo.YaleBrightStarID = CelestialStarInputData.YaleBrightStarID;
 								StarInfo.ColorIndex = CelestialStarInputData.ColorIndex;
 
-								double UEBodyDistance = CelestialVaultDistance * UE_ONE_KILOMETER * StarsVaultPercentage / 100.0;
-								StarInfo.ComputeUELocalTransform(1.0f, UEBodyDistance);
+								// Don't compute the UE Local Transform but use the one we use for the ISM component. 
+								StarInfo.UELocalTransform = FTransform(StarLocation);
 								// Maybe add other computed values here... 
 								Stars.Add(StarInfo);
 							}
@@ -783,7 +873,7 @@ void ACelestialVaultDaySequenceActor::InitStars()
 						if (StarInputData.Magnitude <= MaxVisibleMagnitude)
 						{
 							// Location - Convert to UE Left-handed Frame (Invert the Y coordinate)
-							FVector StarLocation = UCelestialMaths::RADECToXYZ_RH(StarInputData.RA * 15.0, StarInputData.DEC, CelestialVaultDistance * 1000.0 * StarsVaultPercentage / 100.0) * FVector(100.0, -100.0, 100.0);
+							FVector StarLocation = UCelestialMaths::RADECToXYZ_RH(StarInputData.RA * 15.0, StarInputData.DEC, StarsVaultPercentage) * FVector(1.0, -1.0, 1.0);
 
 							// Color
 							FLinearColor StarColor = StarInputData.Color;
@@ -807,7 +897,8 @@ void ACelestialVaultDaySequenceActor::InitStars()
 								StarInfo.Name = StarInputData.Name;
 								StarInfo.Magnitude = StarInputData.Magnitude;
 								StarInfo.Color = StarColor;
-								StarInfo.ComputeUELocalTransform();
+								// Don't compute the UE Local Transform but use the one we use for the ISM component. 
+								StarInfo.UELocalTransform = FTransform(StarLocation);
 								Stars.Add(StarInfo);
 							}
 						}
@@ -821,28 +912,15 @@ void ACelestialVaultDaySequenceActor::InitStars()
 	}		
 }
 
-void ACelestialVaultDaySequenceActor::InitPlanetaryBodies()
+void ACelestialVaultDaySequenceActor::InitPlanetsComponent()
 {
-	TRACE_CPUPROFILER_EVENT_SCOPE(ACelestialVaultDaySequenceActor::InitPlanetaryBodies);
+	TRACE_CPUPROFILER_EVENT_SCOPE(InitPlanetsComponent);
 	
 	// Get the proper Julian day, at Midnight... The Daysequence will rotate the sky vault later 
 	FDateTime LocalTime0 = GetDate();
 	FDateTime UTCTime0 = UCelestialMaths::LocalTimeToUTCTime(LocalTime0, GMT_TimeZone, IsDaylightSavingsNow());
 	double JulianDay = UCelestialMaths::UTCDateTimeToJulianDate(UTCTime0);
 
-	// Init the sun
-	Sun = GetSun(JulianDay);
-	
-	// Init the Moon
-	if (MoonDiscComponent)
-	{
-		// We don't want to use the GetMoonTransform function here because we also want to keep the Moon information for the day.
-		// It has a pitfall because the Moon information will not be adjusted to the current TOD, but that doesn't prevent anyone for doing a manuel query at some point
-		Moon = GetMoon(JulianDay);
-		SetMoonDiscAge(Moon.Age);
-		MoonDiscComponent->SetRelativeTransform(Moon.UELocalTransform);
-	}
-	
 	// Init the Planets from the catalog
 	if (PlanetsComponent)
 	{
@@ -864,12 +942,18 @@ void ACelestialVaultDaySequenceActor::InitPlanetaryBodies()
 			return;
 		}
 
+		// Compute Earth Reference location
+		FPlanetaryBodyKinematicState EarthKinematicState = UCelestialMaths::GetPlanetaryBodyKinematicState_AU(JulianDay, EVSOP87BodyType::Earth);
+
+		
 		// Init from the DataTable
-		PlanetsCatalog->ForeachRow<FPlanetaryBodyInputData>("AEarthDaySequenceActor::InitPlanetaryBodies", [&](const FName& Key, const FPlanetaryBodyInputData& InputPlanetaryBody)
+		PlanetsCatalog->ForeachRow<FPlanetaryBodyInputData>("AEarthDaySequenceActor::InitPlanetsComponent", [&](const FName& Key, const FPlanetaryBodyInputData& InputPlanetaryBody)
 			{
+				// Compute the Body location relative to the Earth reference state for the JuliandDay. 
+				FPlanetaryBody BodyInfo = GetPlanetaryBodyInfo_UsingKnownState(JulianDay, InputPlanetaryBody, EarthKinematicState);
+				
+				// Compute the UE Transform tp locate the Body on the Sphere
 				double UEBodyDistance = CelestialVaultDistance * UE_ONE_KILOMETER * PlanetsVaultPercentage / 100.0;
-				FPlanetaryBody BodyInfo = GetPlanetaryBodyInfo(InputPlanetaryBody, JulianDay);
-				// Compute the Transform
 				BodyInfo.ComputeUELocalTransform(PlanetsScale, UEBodyDistance);
 			
 				// Add the new ISM Instance
@@ -892,7 +976,7 @@ void ACelestialVaultDaySequenceActor::InitPlanetaryBodies()
 	}
 }
 
-FPlanetaryBody ACelestialVaultDaySequenceActor::GetPlanetaryBodyInfo(const FPlanetaryBodyInputData& InputPlanetaryBody, double JulianDay) const
+FPlanetaryBody ACelestialVaultDaySequenceActor::GetPlanetaryBodyInfo_UsingKnownState(double JulianDay, const FPlanetaryBodyInputData& InputPlanetaryBody, const FPlanetaryBodyKinematicState& EarthKinematicState, bool bIgnoreRelativisticEffect) const
 {
 	FPlanetaryBody BodyInfo;
 	BodyInfo.OrbitType = InputPlanetaryBody.OrbitType;
@@ -902,7 +986,7 @@ FPlanetaryBody ACelestialVaultDaySequenceActor::GetPlanetaryBodyInfo(const FPlan
 	// Compute location	
 	double RAJ2000Hours, DECJ2000Degrees, RAHours, DECDegrees, DistanceToEarthAU, DistanceToSunAU, DistanceEarthToSunAU;
 	
-	UCelestialMaths::GetBodyCelestialCoordinatesAU(JulianDay, InputPlanetaryBody.VSOP87BodyType, Latitude, Longitude, RAJ2000Hours, DECJ2000Degrees,RAHours, DECDegrees, DistanceToEarthAU, DistanceToSunAU, DistanceEarthToSunAU);
+	UCelestialMaths::GetBodyCelestialCoordinatesAU_UsingKnownState(JulianDay, InputPlanetaryBody.VSOP87BodyType, EarthKinematicState, Latitude, Longitude, bIgnoreRelativisticEffect, RAJ2000Hours, DECJ2000Degrees,RAHours, DECDegrees, DistanceToEarthAU, DistanceToSunAU, DistanceEarthToSunAU);
 	BodyInfo.RAJ2000 = RAJ2000Hours;
 	BodyInfo.DECJ2000 = DECJ2000Degrees;
 	BodyInfo.RA = RAHours;
@@ -947,7 +1031,124 @@ void ACelestialVaultDaySequenceActor::RebuildAll()
 		UpdatePlanetCenter();
 
 		// Rebuild Sky and Sequence
-		InitStars();
-		InitPlanetaryBodies();
+		InitStarsComponent();
+		InitPlanetsComponent();
+		ResetBodyMotionCache();
+	}
+}
+
+void ACelestialVaultDaySequenceActor::UpdateBodiesMotion()
+{
+	TRACE_CPUPROFILER_EVENT_SCOPE(UpdateCelestialBodiesMotion);
+	
+	double TimeOfDay = GetTimeOfDay();
+	TimeOfDay = FMath::Fmod(TimeOfDay, 24.0); // After incrementing the number of loops, the Sequence player returns a TOD that could be 24. Make a FMod to make sure we are at the beginning of the next day  
+		
+	EnsureKeysAroundTimeOfDay(TimeOfDay); 
+	
+	if (MoonDiscComponent)
+	{
+		MoonDiscComponent->SetRelativeLocation(MoonDiscLocationCurve.Eval(TimeOfDay));
+		MoonDiscComponent->SetRelativeRotation(MoonDiscRotationCurve.Eval(TimeOfDay));
+		SetMoonDiscAge(MoonAgeCurve.Eval(TimeOfDay));
+	}
+
+	if ( MoonLightComponent)
+	{
+		MoonLightComponent->SetRelativeRotation(MoonLightRotationCurve.Eval(TimeOfDay));
+	}
+
+	if ( SunLightComponent)
+	{
+		SunLightComponent->SetRelativeRotation(SunLightRotationCurve.Eval(TimeOfDay));
+	}
+
+	// Don't use a curve cache for the Celestial Vault rotation. It's a basic angular rotation. 
+	double VaultYaw = GetCelestialVaultAngle(SequenceLoopCount * 24 + TimeOfDay );
+	CelestialVaultComponent->SetRelativeRotation(FRotator(0.0, VaultYaw, 0.0));
+}
+
+/// Bodies Motion Cache
+void ACelestialVaultDaySequenceActor::ResetBodyMotionCache()
+{
+	TRACE_CPUPROFILER_EVENT_SCOPE(BuildCurvesCache);
+	UE_LOG(LogCelestialVault, Verbose, TEXT("Reset Body Motion Cache for Day %d"), SequenceLoopCount);
+
+	MoonAgeCurve.Reset();
+	MoonDiscLocationCurve.Reset();
+	MoonDiscRotationCurve.Reset();
+	MoonLightRotationCurve.Reset();
+	SunLightRotationCurve.Reset();
+
+	BodyMotionCacheMax = 0.0;
+	BodyMotionCacheMin = 0.0;
+}
+
+void ACelestialVaultDaySequenceActor::EnsureKeysAroundTimeOfDay(double TimeOfDay)
+{
+	TRACE_CPUPROFILER_EVENT_SCOPE(EnsureKeysAroundTimeOfDay);
+	
+	// The Cache will be split into even Time Step intervals, using BodyMotionCacheKeyFramesCount
+	const double TimeStep = 24.0 / FMath::Max(static_cast<signed>(BodyMotionCacheKeyFramesCount), 1);
+
+	double LeftKeyValue = FMath::Floor(TimeOfDay / TimeStep) * TimeStep;
+	double RightKeyValue = LeftKeyValue + TimeStep;
+
+	// Nothing in cache, add two points
+	if (BodyMotionCacheMin == 0 && BodyMotionCacheMax == 0)
+	{
+		// First query, we don't have KeyFrames yet, so add them
+		AddBodyMotionCacheKey(LeftKeyValue);
+		AddBodyMotionCacheKey(RightKeyValue);
+	}
+
+	// If needed, Extend left: add ALL intermediate keys at constant Step (this allows for arbitrary jumps in time of day)
+	while (LeftKeyValue < BodyMotionCacheMin)
+	{
+		AddBodyMotionCacheKey(BodyMotionCacheMin - TimeStep);
+	}
+
+	// Same extention to the Right. 
+	while (RightKeyValue > BodyMotionCacheMax)
+	{
+		AddBodyMotionCacheKey(BodyMotionCacheMax + TimeStep);
+	}
+}
+
+void ACelestialVaultDaySequenceActor::AddBodyMotionCacheKey(double TimeOfDay)
+{
+	TRACE_CPUPROFILER_EVENT_SCOPE(AddBodyMotionCacheKey);
+	UE_LOG(LogCelestialVault, Verbose, TEXT("Adding new Keyframe in the BodyMotion Cache, for TOD = %f"), TimeOfDay);
+	
+	// We need the Julian Date for the celestial computation. For this, we need to consider the initial day, the number of times the sequence has looped, and the current time of day.
+	// Compute the Local time
+	FDateTime LocalTime = GetDate() + FTimespan::FromHours(TimeOfDay);
+	
+	// Convert DateTime to Julian Date 
+	FDateTime UTCTime = UCelestialMaths::LocalTimeToUTCTime(LocalTime, GMT_TimeZone, IsDaylightSavingsAtTimeOfDay(TimeOfDay));
+	double JulianDate = UCelestialMaths::UTCDateTimeToJulianDate(UTCTime);
+
+	// Compute Earth Reference location - We'll use it as a cache for Moon/Sun faster computations
+	FPlanetaryBodyKinematicState EarthKinematicState = UCelestialMaths::GetPlanetaryBodyKinematicState_AU(JulianDate, EVSOP87BodyType::Earth);
+		
+	// Get Moon data
+	FPlanetaryBody  MoonInfo = ComputeMoonInfo_UsingKnownState(JulianDate, EarthKinematicState);
+	MoonAgeCurve.AddPoint(TimeOfDay, MoonInfo.Age);
+	MoonDiscLocationCurve.AddPoint(TimeOfDay, MoonInfo.UELocalTransform.GetLocation());
+	MoonDiscRotationCurve.AddPoint(TimeOfDay, MoonInfo.UELocalTransform.GetRotation());
+	MoonLightRotationCurve.AddPoint(TimeOfDay, FQuat::FindBetweenVectors(FVector::ForwardVector, MoonInfo.DirectionTowardEarth));
+	// Scale won't change over time but at least make sure we apply it once
+	if (MoonDiscComponent)
+	{
+		MoonDiscComponent->SetRelativeScale3D(MoonInfo.UELocalTransform.GetScale3D());	
 	}
+	
+	// Get Sun Data
+	FStellarBody SunInfo = ComputeSunInfo_UsingKnownState(JulianDate, EarthKinematicState);
+	SunLightRotationCurve.AddPoint(TimeOfDay, FQuat::FindBetweenVectors(FVector::ForwardVector, SunInfo.DirectionTowardEarth));
+
+	// Update our known interval cache. 
+	BodyMotionCacheMax = FMath::Max(TimeOfDay, BodyMotionCacheMax);
+	BodyMotionCacheMin = FMath::Min(TimeOfDay, BodyMotionCacheMin);
 }
+
diff --git a/Engine/Plugins/Experimental/CelestialVault/Source/CelestialVault/Private/CelestialVaultSequence.cpp b/Engine/Plugins/Experimental/CelestialVault/Source/CelestialVault/Private/CelestialVaultSequence.cpp
deleted file mode 100644
diff --git a/Engine/Plugins/Experimental/CelestialVault/Source/CelestialVault/Private/DaylightSavings.cpp b/Engine/Plugins/Experimental/CelestialVault/Source/CelestialVault/Private/DaylightSavings.cpp
--- a/Engine/Plugins/Experimental/CelestialVault/Source/CelestialVault/Private/DaylightSavings.cpp
+++ b/Engine/Plugins/Experimental/CelestialVault/Source/CelestialVault/Private/DaylightSavings.cpp
@@ -152,14 +152,28 @@ bool UDaylightSavings::IsDaylightSavings(FDateTime DateTime, EDaylightSavingsMod
 
 	int32 EndTimeYear = DateTime.GetYear();
 
-	if (Mode == EDaylightSavingsMode::Southern ||
-		(Mode == EDaylightSavingsMode::Latitude && Latitude < 0))
+	if (Mode == EDaylightSavingsMode::SouthernHemisphere ||
+		(Mode == EDaylightSavingsMode::Automatic && Latitude < 0))
 	{
 		EndTimeYear++;
 	}
 	
-	FDateTime StartTime = ToDate(StartDay, DateTime.GetYear()) + FTimespan::FromHours(SwitchHour);
-	FDateTime EndTime = ToDate(EndDay,  EndTimeYear) + FTimespan::FromHours(SwitchHour);
+	FDateTime StartTime = ToDate(StartDay, DateTime.GetYear());
+	if (StartTime == FDateTime::MinValue())
+	{
+		// invalid Date - return false by default
+		return false;
+	}
+	StartTime += FTimespan::FromHours(SwitchHour);
+
+
+	FDateTime EndTime = ToDate(EndDay,  EndTimeYear);
+	if (EndTime == FDateTime::MinValue())
+	{
+		// invalid Date - return false by default
+		return false;
+	}
+	EndTime += FTimespan::FromHours(SwitchHour);
 
 	return (DateTime >= StartTime && DateTime <= EndTime);
 }
diff --git a/Engine/Plugins/Experimental/CelestialVault/Source/CelestialVault/Public/CelestialDataTypes.h b/Engine/Plugins/Experimental/CelestialVault/Source/CelestialVault/Public/CelestialDataTypes.h
--- a/Engine/Plugins/Experimental/CelestialVault/Source/CelestialVault/Public/CelestialDataTypes.h
+++ b/Engine/Plugins/Experimental/CelestialVault/Source/CelestialVault/Public/CelestialDataTypes.h
@@ -161,4 +161,24 @@ struct FPlanetaryBody : public FCelestialBody
 };
 
 
+/**
+ * Runtime structure to store the Kinematic State - Location, Velocity of a Planetary Body, in VSOP87 and FK5J2000 frames
+ */
+USTRUCT(BlueprintInternalUseOnly)
+struct FPlanetaryBodyKinematicState
+{
+	GENERATED_BODY()
+
+	UPROPERTY(BlueprintReadWrite, EditAnywhere, Category= "Celestial Vault")
+	FVector Location_VSOP87 = FVector::ZeroVector;
+
+	UPROPERTY(BlueprintReadWrite, EditAnywhere, Category= "Celestial Vault")
+	FVector Velocity_VSOP87 = FVector::ZeroVector;
+
+	UPROPERTY(BlueprintReadWrite, EditAnywhere, Category= "Celestial Vault")
+	FVector Location_Heliocentric_FK5J2000 = FVector::ZeroVector;
+
+	UPROPERTY(BlueprintReadWrite, EditAnywhere, Category= "Celestial Vault")
+	FVector Velocity_FK5J2000 = FVector::ZeroVector;
+};
 
diff --git a/Engine/Plugins/Experimental/CelestialVault/Source/CelestialVault/Public/CelestialMaths.h b/Engine/Plugins/Experimental/CelestialVault/Source/CelestialVault/Public/CelestialMaths.h
--- a/Engine/Plugins/Experimental/CelestialVault/Source/CelestialVault/Public/CelestialMaths.h
+++ b/Engine/Plugins/Experimental/CelestialVault/Source/CelestialVault/Public/CelestialMaths.h
@@ -74,24 +74,49 @@ public:
      *       
      * The returned Location and the Observer Body location are expressed in Astronomical Units (AU)
      * The relativistic effects are considered (time taken for the light to reach the Observer body location)
-     * This function therefore returns the location of the Planetary Body as if it was seen from the Observer Body Location
+     * This function therefore returns the location of the Planetary Body as if it was seen from the Observer Body Location at this JulianDate
     */
     UFUNCTION(BlueprintCallable, Category="Celestial Vault|VSOP87")
     static FVector GetSolarSystemBodyLocation_VSOP87_Relativistic(FVector ReferenceBodyLocation_Heliocentric_AU, EVSOP87BodyType VSOP87BodyType, double JulianDate);
 
-	
+	/**
+	* Returns the Body location and velocity in the VSOP87 and FK5J2000 frames
+	*
+	*	This is a convenience function to keep a cache of a Body location (eg. Earth) and accelerate the computations when interested in several other bodies. 
+	*/
+	UFUNCTION(BlueprintCallable, Category="Celestial Vault|VSOP87")
+	static FPlanetaryBodyKinematicState GetPlanetaryBodyKinematicState_AU(double JulianDate, EVSOP87BodyType VSOP87BodyType);
+
 #pragma endregion
 
 #pragma region Celestial Bodies
-
+	
     /** Return the location of a Planetary Body relative to the Earth, expressed in Celestial Coordinates (RA, DEC, Distance)
      * It requires the Observer location on Earth for more precise computations
      * This function also returns the distance between bodies, as it can help for Magnitude computations
+     *
+	 * By default, the relativistic effects are considered (time taken for the light to reach the Observer body location)
+	 * This function therefore returns the location of the Planetary Body as if it was seen from the Observer Body Location at this JulianDate
+	 * With bIgnoreRelativisticEffect=true, one can ignore this and have the instant location. It can also save some computation cycles for close bodies like Moon. 
      * **/
     UFUNCTION(BlueprintCallable, Category="Celestial Vault|Planetary bodies")
-    static void GetBodyCelestialCoordinatesAU(double JulianDate, EVSOP87BodyType VSOP87BodyType, double ObserverLatitude, double ObserverLongitude,
+    static void GetBodyCelestialCoordinatesAU(double JulianDate, EVSOP87BodyType VSOP87BodyType, double ObserverLatitude, double ObserverLongitude, bool bIgnoreRelativisticEffect,
     	double& RAJ2000Hours, double& DECJ2000Degrees, double& RAHours, double& DECDegrees, double& DistanceBodyToEarthAU, double& DistanceBodyToSunAU, double& DistanceEarthToSunAU);
 
+	/** Return the location of a Planetary Body relative to the Earth, expressed in Celestial Coordinates (RA, DEC, Distance)
+	 * It requires the Observer location on Earth for more precise computations
+	 * This function also returns the distance between bodies, as it can help for Magnitude computations
+	 *
+	 * By default, the relativistic effects are considered (time taken for the light to reach the Observer body location)
+	 * This function therefore returns the location of the Planetary Body as if it was seen from the Observer Body Location at this JulianDate
+	 * With bIgnoreRelativisticEffect=true, one can ignore this and have the instant location. It can also save some computation cycles for close bodies like Moon. 
+	 *
+	 * This override can run faster if we provide it with a cache of the Earth Location and Velocity, because it won't recompute it. 
+	 * **/
+	UFUNCTION(BlueprintCallable, Category="Celestial Vault|Planetary bodies")
+	static void GetBodyCelestialCoordinatesAU_UsingKnownState(double JulianDate, EVSOP87BodyType VSOP87BodyType, const FPlanetaryBodyKinematicState& EarthKinematicState, double ObserverLatitude, double ObserverLongitude, bool bIgnoreRelativisticEffect,
+		double& RAJ2000Hours, double& DECJ2000Degrees, double& RAHours, double& DECDegrees, double& DistanceBodyToEarthAU, double& DistanceBodyToSunAU, double& DistanceEarthToSunAU);
+	
     /** Return the Magnitude of a Planetary Body as seen from the Earth **/
     UFUNCTION(BlueprintCallable, Category="Celestial Vault|Planetary bodies")
     static double GetPlanetaryBodyMagnitude(EVSOP87BodyType VSOP87BodyType, double DistanceToSunAU, double DistanceToEarthAU, double DistanceEarthToSunAU, double& PhaseAngle);
@@ -317,6 +342,14 @@ public:
     UFUNCTION(BlueprintCallable, Category="Celestial Vault|Sun")
     static FStellarBody GetSunInformation(double JulianDate, double ObserverLatitude, double ObserverLongitude);
 
+	/** Compute all Sun Properties for a specific JulianDate
+	 *
+	 * This override can run faster if we provide it with a cache of the Earth Location and Velocity, because it won't recompute it. 
+	 */ 
+	UFUNCTION(BlueprintCallable, Category="Celestial Vault|Sun")
+	static FStellarBody GetSunInformation_UsingKnownState(double JulianDate, const FPlanetaryBodyKinematicState& EarthKinematicState, double ObserverLatitude, double ObserverLongitude);
+
+	
 	/** Returns the Moon Phase for a specific Date
 	 * 
 	 *  This is an approximate computation using a number of lunar cycles with a synodic month equals to 29.53059 days
diff --git a/Engine/Plugins/Experimental/CelestialVault/Source/CelestialVault/Public/CelestialVaultDaySequenceActor.h b/Engine/Plugins/Experimental/CelestialVault/Source/CelestialVault/Public/CelestialVaultDaySequenceActor.h
--- a/Engine/Plugins/Experimental/CelestialVault/Source/CelestialVault/Public/CelestialVaultDaySequenceActor.h
+++ b/Engine/Plugins/Experimental/CelestialVault/Source/CelestialVault/Public/CelestialVaultDaySequenceActor.h
@@ -109,13 +109,17 @@ public: // Properties
 	UPROPERTY(EditAnywhere, BlueprintReadWrite, Category="Date", meta = (EditConditionHides, EditCondition = "!bUseCurrentDate", ClampMin = "1", ClampMax = "31"))
 	int Day = 7;
 
+	/** At runtime, if we cycle over more than one day, we increase this one */
+	UPROPERTY(Transient, VisibleAnywhere, BlueprintReadOnly, Category="Date")
+	int SequenceLoopCount = 0;
+	
 #pragma endregion
 
 #pragma region Daylight Saving related Properties
     
 	/** Define how the Daylight Savings are taken into consideration */ 
 	UPROPERTY(EditAnywhere, BlueprintReadWrite, Category = "Daylight Savings")
-	EDaylightSavingsMode DaylightSavingsMode = EDaylightSavingsMode::Latitude;
+	EDaylightSavingsMode DaylightSavingsMode = EDaylightSavingsMode::Automatic;
 	
 	/** Rule that applies to determine the Daylight Saving stating day */ 
 	UPROPERTY(EditAnywhere, BlueprintReadWrite, Category="Daylight Savings", meta=(EditConditionHides, EditCondition = "DaylightSavingsMode!=EDaylightSavingsMode::None", DisplayName = "Starts On"))
@@ -243,10 +247,6 @@ public: // Properties
 	UPROPERTY(EditAnywhere, BlueprintReadWrite, Category="Moons", meta = (ClampMin = "0"))
 	float MoonScale = 2.0f;
 
-	/** Celestial Info for the Moon, at the beginning of the Day */  
-	UPROPERTY(BlueprintReadOnly, Category="Moons")
-	FPlanetaryBody Moon;
-
 	/** If true, the moon Age (Phase) and location can be overriden */  
 	UPROPERTY(EditAnywhere, Category="Moons")
 	bool bManualControl = false;
@@ -271,10 +271,6 @@ public: // Properties
 	UPROPERTY(EditAnywhere, Category="Moons", meta = (EditConditionHides, DisplayName="Moon Offset to Sun's Declination", EditCondition = "bManualControl==true", ForceUnits="Degrees", ClampMin = "-45", ClampMax = "45"))
 	float MoonOffset_DEC = 15.0f;
 
-	/** Celestial Info for the Sun, at the beginning of the Day */  
-	UPROPERTY(BlueprintReadOnly, Category="Sun")
-	FStellarBody Sun; 
-
 	/** Base Sun Intensity
 	 * Typically 120000 Lux
 	 */
@@ -293,6 +289,10 @@ protected: // Functions
 	/** BeginPlay and OnConstruction overrides auto-register this actor with the DaySequenceSubsystem. */
 	virtual void BeginPlay() override;
 	virtual void OnConstruction(const FTransform& Transform) override;
+	virtual void Tick(float DeltaTime) override;
+
+	/** Called when the sequence is updated */ 
+	virtual void SequencePlayerUpdated(float CurrentTime, float PreviousTime) override;
 
 public:
 	virtual void PostRegisterAllComponents() override;
@@ -307,15 +307,35 @@ public: // Functions
 	/** Returns the current defined day, with the Time defined by the DaySequence Time of day*/ 
 	UFUNCTION(BlueprintCallable, Category = "Celestial Vault|Date")
 	FDateTime GetDateAndTime() const;
+
+	/** Returns the current defined day, with the Time defined by the DaySequence Time of day*/
+	UFUNCTION(BlueprintCallable, Category = "Celestial Vault|Date")
+	double GetJulianDate() const;
 	
 	/** Returns the Celestial Info for the Sun, at a specific Julian Date */  
 	UFUNCTION(BlueprintCallable, Category = "Celestial Vault|Sun")
-	FStellarBody GetSun(double JulianDate) const; 
+	FStellarBody ComputeSunInfo(double JulianDate) const; 
+
+	/** Returns the Celestial Info for the Sun, at current the Date and Time */  
+	UFUNCTION(BlueprintCallable, Category = "Celestial Vault|Sun")
+	FStellarBody ComputeCurrentSunInfo() const; 
+
+	/** Returns the Celestial Info for the Sun, at a specific Julian Date */  
+	UFUNCTION(BlueprintCallable, Category = "Celestial Vault|Sun")
+	FStellarBody ComputeSunInfo_UsingKnownState(double JulianDate, const FPlanetaryBodyKinematicState& EarthKinematicState) const;
 	
 	/** Returns the Celestial Info for the Moon, at a specific Julian Date */  
-	UFUNCTION(BlueprintCallable, Category = "Celestial Vault|Queries")
-	FPlanetaryBody GetMoon(double JulianDate) const;
+	UFUNCTION(BlueprintCallable, Category = "Celestial Vault|Moons")
+	FPlanetaryBody ComputeMoonInfo(double JulianDate) const;
+
+	/** Returns the Celestial Info for the Moon, at current the Date and Time */  
+	UFUNCTION(BlueprintCallable, Category = "Celestial Vault|Moons")
+	FPlanetaryBody ComputeCurrentMoonInfo() const;
 
+	/** Returns the Celestial Info for the Moon, at a specific Julian Date */  
+	UFUNCTION(BlueprintCallable, Category = "Celestial Vault|Moons")
+	FPlanetaryBody ComputeMoonInfo_UsingKnownState(double JulianDate, const FPlanetaryBodyKinematicState& EarthKinematicState) const;
+	
 	/** Manually set the Moon Age (Phase) */  
 	UFUNCTION(BlueprintCallable, Category = "Celestial Vault|Moons")
 	void SetMoonDiscAge(float InMoonAge);
@@ -346,14 +366,44 @@ public: // Functions
 #endif
 
 private:
-	void InitStars();
-	void InitPlanetaryBodies();
-	FPlanetaryBody GetPlanetaryBodyInfo(const FPlanetaryBodyInputData& InputPlanetaryBody, double JulianDay) const;
+	void InitStarsComponent();
+	void InitPlanetsComponent();
+	FPlanetaryBody GetPlanetaryBodyInfo_UsingKnownState(double JulianDay, const FPlanetaryBodyInputData& InputPlanetaryBody, const FPlanetaryBodyKinematicState& EarthKinematicState, bool bIgnoreRelativisticEffect = false) const;
 	void UpdatePlanetCenter();
 
 	UFUNCTION(BlueprintCallable, CallInEditor, Category= "Celestial Vault")
 	void RebuildAll();
 
+	// Special behavior for celestial bodies motion: 
+	// The DaySequence system is designed to loop over one single day.
+	// In our case, we want to simulate more than one day, and have for instance the moon properly moving in the sky over midnight.
+	// So we won't use a procedural sequence for the bodies state, but move the body location on demand, when the DaySequence is updated.
+	// To maintain performance, we'll use a local motion cache updated when needed. 
+
+	/** Moves the celestial bodies when the DaySequence is updated
+	 *  The motion cache is automatically build when a state is queried, ensuring we have the proper keyframe around the current time of day 
+	 * **/ 
+	void UpdateBodiesMotion();
+
+	/** Clear the Bodies Motion Cache for the relevant objects - Sun, Moon, Vault **/
+	void ResetBodyMotionCache();
+
+	/** Check if we have the proper keyframes around the desired Time of Day. If not, it will add them */ 
+	void EnsureKeysAroundTimeOfDay(double TimeOfDay);
+
+	/** Compute and Add a new KeyFrame in the cache */ 
+	void AddBodyMotionCacheKey(double CurveKeyTimeHour);
+
+	// Local Interpolable Curves Cache for runtime-moving bodies
+	const unsigned int BodyMotionCacheKeyFramesCount = 24;
+	double BodyMotionCacheMin = 0.0;
+	double BodyMotionCacheMax = 0.0;
+	FInterpCurve<FVector>  MoonDiscLocationCurve;
+	FInterpCurve<FQuat>  MoonDiscRotationCurve;
+	FInterpCurve<FQuat>  MoonLightRotationCurve;
+	FInterpCurve<float>  MoonAgeCurve;
+	FInterpCurve<FQuat>  SunLightRotationCurve;
+	
 #if CELESTIAL_VAULT_ENABLE_DRAW_DEBUG
 
 	const FString CelestialVaultSF = TEXT("CelestialVault");
diff --git a/Engine/Plugins/Experimental/CelestialVault/Source/CelestialVault/Public/CelestialVaultSequence.h b/Engine/Plugins/Experimental/CelestialVault/Source/CelestialVault/Public/CelestialVaultSequence.h
deleted file mode 100644
diff --git a/Engine/Plugins/Experimental/CelestialVault/Source/CelestialVault/Public/DaylightSavings.h b/Engine/Plugins/Experimental/CelestialVault/Source/CelestialVault/Public/DaylightSavings.h
--- a/Engine/Plugins/Experimental/CelestialVault/Source/CelestialVault/Public/DaylightSavings.h
+++ b/Engine/Plugins/Experimental/CelestialVault/Source/CelestialVault/Public/DaylightSavings.h
@@ -11,10 +11,14 @@
 UENUM(BlueprintType)
 enum class EDaylightSavingsMode : uint8
 {
-	None,       // Do not apply daylight saving time (DST).
-	Latitude,   // Determine hemisphere from the latitude sign (>= 0 = Northern, < 0 = Southern).
-	Northern,   // Force Northern Hemisphere rules (DST typically observed during Apr–Oct, i.e., around June).
-	Southern    // Force Southern Hemisphere rules (DST typically observed during Oct–Mar, i.e., around December).
+	// Do not apply daylight saving time (DST).
+	None,
+	// Automatically determine the Hemisphere from the Latitude sign (>= 0 = Northern, < 0 = Southern). Consider other N/S hemisphere options if your Origin is Geocentric or different from the logic you want to apply. 
+	Automatic,
+	// Northern Hemisphere rules (DST typically observed during Apr–Oct, i.e., around June).
+	NorthernHemisphere,
+	// Southern Hemisphere rules (DST typically observed during Oct–Mar, i.e., around December). The following year is used to compute the end date
+	SouthernHemisphere    
 };
 
 /** Different Daylight Savings rules that apply across the globe*/
diff --git a/Engine/Plugins/Experimental/CelestialVault/Source/CelestialVaultEditor/Private/CelestialVaultTests.cpp b/Engine/Plugins/Experimental/CelestialVault/Source/CelestialVaultEditor/Private/CelestialVaultTests.cpp
--- a/Engine/Plugins/Experimental/CelestialVault/Source/CelestialVaultEditor/Private/CelestialVaultTests.cpp
+++ b/Engine/Plugins/Experimental/CelestialVault/Source/CelestialVaultEditor/Private/CelestialVaultTests.cpp
@@ -346,13 +346,13 @@ bool CelestialVault_RADEC::RunTest(const FString& Parameters)
 			double DistanceBodyToSunAU = 0.0;
 			double DistanceEarthToSunAU = 0.0;
 
-			// Get expected values
-			UCelestialMaths::GetBodyCelestialCoordinatesAU(Case.JulianDate, BodyCases.VSOP87BodyType, Case.Latitude, Case.Longitude, RAJ2000Hours, DECJ2000Degrees, RAHours, DECDegrees, DistanceBodyToEarthAU, DistanceBodyToSunAU, DistanceEarthToSunAU);
+			// Compute Values
+			UCelestialMaths::GetBodyCelestialCoordinatesAU(Case.JulianDate, BodyCases.VSOP87BodyType, Case.Latitude, Case.Longitude, false,  RAJ2000Hours, DECJ2000Degrees, RAHours, DECDegrees, DistanceBodyToEarthAU, DistanceBodyToSunAU, DistanceEarthToSunAU);
 			double RAJ2000Degrees = RAJ2000Hours * 15;
 			double RADegrees = RAHours * 15;
 
 			// Check RAJ2000
-			Error = FMath::Abs(RAJ2000Degrees - Case.ExpectedRAJ2000Degrees) / Case.ExpectedRAJ2000Degrees * 100.0f;
+			Error = FMath::Abs((RAJ2000Degrees - Case.ExpectedRAJ2000Degrees) / FMath::Max(Case.ExpectedRAJ2000Degrees, KINDA_SMALL_NUMBER)) * 100.0;
 			++AverageErrorSampleCount;
 			AverageError += (Error - AverageError) / AverageErrorSampleCount;
 			MaxError = FMath::Max(MaxError, Error);
@@ -363,7 +363,7 @@ bool CelestialVault_RADEC::RunTest(const FString& Parameters)
 			}
 
 			// Check DECJ2000
-			Error = FMath::Abs(DECJ2000Degrees - Case.ExpectedDecJ2000Degrees) / Case.ExpectedDecJ2000Degrees * 100.0f;
+			Error = FMath::Abs((DECJ2000Degrees - Case.ExpectedDecJ2000Degrees) /  FMath::Max(Case.ExpectedDecJ2000Degrees, KINDA_SMALL_NUMBER)) * 100.0;
 			++AverageErrorSampleCount;
 			AverageError += (Error - AverageError) / AverageErrorSampleCount;
 			MaxError = FMath::Max(MaxError, Error);
@@ -374,7 +374,7 @@ bool CelestialVault_RADEC::RunTest(const FString& Parameters)
 			}
 
 			// Check RA at Date
-			Error = FMath::Abs(RADegrees - Case.ExpectedRADegrees) / Case.ExpectedRADegrees * 100.0f;
+			Error = FMath::Abs((RADegrees - Case.ExpectedRADegrees) /  FMath::Max(Case.ExpectedRADegrees, KINDA_SMALL_NUMBER)) * 100.0;
 			++AverageErrorSampleCount;
 			AverageError += (Error - AverageError) / AverageErrorSampleCount;
 			MaxError = FMath::Max(MaxError, Error);
@@ -385,7 +385,7 @@ bool CelestialVault_RADEC::RunTest(const FString& Parameters)
 			}
 
 			// Check DEC at Date
-			Error = FMath::Abs(DECDegrees - Case.ExpectedDecDegrees) / Case.ExpectedDecDegrees * 100.0f;
+			Error = FMath::Abs((DECDegrees - Case.ExpectedDecDegrees) /  FMath::Max(Case.ExpectedDecDegrees, KINDA_SMALL_NUMBER)) * 100.0;
 			++AverageErrorSampleCount;
 			AverageError += (Error - AverageError) / AverageErrorSampleCount;
 			MaxError = FMath::Max(MaxError, Error);
diff --git a/Engine/Plugins/Experimental/CelestialVault/Source/CelestialVaultEditor/Private/DaylightSavingsCustomization.cpp b/Engine/Plugins/Experimental/CelestialVault/Source/CelestialVaultEditor/Private/DaylightSavingsCustomization.cpp
--- a/Engine/Plugins/Experimental/CelestialVault/Source/CelestialVaultEditor/Private/DaylightSavingsCustomization.cpp
+++ b/Engine/Plugins/Experimental/CelestialVault/Source/CelestialVaultEditor/Private/DaylightSavingsCustomization.cpp
@@ -43,8 +43,8 @@ void FDaylightSavingsCustomization::CustomizeHeader(TSharedRef<IPropertyHandle>
 				{
 					int YearToUse = CelestialVaultValues.GetValue().Year;
 					
-					if (CelestialVaultValues.GetValue().Mode == EDaylightSavingsMode::Southern ||
-						(CelestialVaultValues.GetValue().Mode == EDaylightSavingsMode::Latitude && CelestialVaultValues.GetValue().Latitude < 0 ))
+					if (CelestialVaultValues.GetValue().Mode == EDaylightSavingsMode::SouthernHemisphere ||
+						(CelestialVaultValues.GetValue().Mode == EDaylightSavingsMode::Automatic && CelestialVaultValues.GetValue().Latitude < 0 ))
 					{
 						// We are in the south hemisphere. But we need to increase the year only for the DaylightSavingsEndDay Property! 
 						if (FProperty* DaylightSavingsRuleProp = LocalHandle->GetProperty())
diff --git a/Engine/Plugins/Experimental/DaySequence/Source/DaySequence/Private/DaySequencePlayer.h b/Engine/Plugins/Experimental/DaySequence/Source/DaySequence/Private/DaySequencePlayer.h
--- a/Engine/Plugins/Experimental/DaySequence/Source/DaySequence/Private/DaySequencePlayer.h
+++ b/Engine/Plugins/Experimental/DaySequence/Source/DaySequence/Private/DaySequencePlayer.h
@@ -130,6 +130,12 @@ public:
 	 */
 	UE_API virtual FQualifiedFrameTime GetDuration() const override;
 
+	/**
+	 * Returns the number of times this sequence has been looped
+	*/
+	UFUNCTION(BlueprintCallable, Category = "Sequencer|Player")
+	UE_API virtual int32 GetCurrentNumLoops() const {return CurrentNumLoops;} ;
+	
 	/**
 	 * Get this sequence's duration in frames
 	 */
diff --git a/Engine/Plugins/Experimental/DaySequence/Source/DaySequence/Public/IDaySequencePlayer.h b/Engine/Plugins/Experimental/DaySequence/Source/DaySequence/Public/IDaySequencePlayer.h
--- a/Engine/Plugins/Experimental/DaySequence/Source/DaySequence/Public/IDaySequencePlayer.h
+++ b/Engine/Plugins/Experimental/DaySequence/Source/DaySequence/Public/IDaySequencePlayer.h
@@ -39,6 +39,8 @@ public:
 	virtual FQualifiedFrameTime GetCurrentTime() const = 0;
 	
 	virtual FQualifiedFrameTime GetDuration() const = 0;
+
+	virtual int32 GetCurrentNumLoops() const = 0;
 	
 	virtual void SetIgnorePlaybackReplication(bool bState) = 0;
 
