Runtime controls
   Added a Stop function to reset the sequence to its initial state.
   Allow for Runtime property changes with dynamic sequence reconstruction.
   Added an example widget to control the properties.

Various fixes
   Added Fog cutoff values for views from space
   Added the ability to use a Geocentri origin.
   Sun and Moon light intensities are now updated when property changes occur.
   Fixed the magnitude computations for Moon, and planets.

#virtualized
#Synced-CL 50847748

diff --git a/Engine/Plugins/Experimental/CelestialVault/Source/CelestialVault/Private/CelestialMaths.cpp b/Engine/Plugins/Experimental/CelestialVault/Source/CelestialVault/Private/CelestialMaths.cpp
--- a/Engine/Plugins/Experimental/CelestialVault/Source/CelestialVault/Private/CelestialMaths.cpp
+++ b/Engine/Plugins/Experimental/CelestialVault/Source/CelestialVault/Private/CelestialMaths.cpp
@@ -151,8 +151,8 @@ FPlanetaryBodyKinematicState UCelestialMaths::GetPlanetaryBodyKinematicState_AU(
 	return BodyKinematicState;
 }
 
-void UCelestialMaths::GetBodyCelestialCoordinatesAU(double JulianDate, EVSOP87BodyType VSOP87BodyType, double ObserverLatitude, double ObserverLongitude, bool bIgnoreRelativisticEffect, 
-                                                    double& RAJ2000Hours, double& DECJ2000Degrees, double& RAHours, double& DECDegrees, double& DistanceBodyToEarthAU, double& DistanceBodyToSunAU, double& DistanceEarthToSunAU) 
+void UCelestialMaths::GetBodyCelestialCoordinatesAU(double JulianDate, EVSOP87BodyType VSOP87BodyType, double ObserverLatitude, double ObserverLongitude, bool bGeoCentric,
+                                                    bool bIgnoreRelativisticEffect, double& RAJ2000Hours, double& DECJ2000Degrees, double& RAHours, double& DECDegrees, double& DistanceBodyToEarthAU, double& DistanceBodyToSunAU, double& DistanceEarthToSunAU) 
 {
 	TRACE_CPUPROFILER_EVENT_SCOPE(GetBodyCelestialCoordinatesAU);
 
@@ -160,10 +160,10 @@ void UCelestialMaths::GetBodyCelestialCoordinatesAU(double JulianDate, EVSOP87Bo
 	FPlanetaryBodyKinematicState EarthKinematicState = GetPlanetaryBodyKinematicState_AU(JulianDate, EVSOP87BodyType::Earth);
 
 	// Call the actual computation method with the Earth as a reference. 
-	GetBodyCelestialCoordinatesAU_UsingKnownState(JulianDate, VSOP87BodyType, EarthKinematicState, ObserverLatitude, ObserverLongitude, bIgnoreRelativisticEffect, RAJ2000Hours, DECJ2000Degrees, RAHours, DECDegrees, DistanceBodyToEarthAU, DistanceBodyToSunAU, DistanceEarthToSunAU);
+	GetBodyCelestialCoordinatesAU_UsingKnownState(JulianDate, VSOP87BodyType, EarthKinematicState, ObserverLatitude, ObserverLongitude, bGeoCentric, bIgnoreRelativisticEffect, RAJ2000Hours, DECJ2000Degrees, RAHours, DECDegrees, DistanceBodyToEarthAU, DistanceBodyToSunAU, DistanceEarthToSunAU);
 }
 
-void UCelestialMaths::GetBodyCelestialCoordinatesAU_UsingKnownState(double JulianDate, EVSOP87BodyType VSOP87BodyType, const FPlanetaryBodyKinematicState& EarthKinematicState, double ObserverLatitude, double ObserverLongitude, bool bIgnoreRelativisticEffect, double& RAJ2000Hours, double& DECJ2000Degrees, double& RAHours, double& DECDegrees, double& DistanceBodyToEarthAU, double& DistanceBodyToSunAU, double& DistanceEarthToSunAU)
+void UCelestialMaths::GetBodyCelestialCoordinatesAU_UsingKnownState(double JulianDate, EVSOP87BodyType VSOP87BodyType, const FPlanetaryBodyKinematicState& EarthKinematicState, double ObserverLatitude, double ObserverLongitude, bool bGeoCentric, bool bIgnoreRelativisticEffect, double& RAJ2000Hours, double& DECJ2000Degrees, double& RAHours, double& DECDegrees, double& DistanceBodyToEarthAU, double& DistanceBodyToSunAU, double& DistanceEarthToSunAU)
 {
 	TRACE_CPUPROFILER_EVENT_SCOPE(GetBodyCelestialCoordinatesAU_UsingKnownState);
 	
@@ -208,10 +208,15 @@ void UCelestialMaths::GetBodyCelestialCoordinatesAU_UsingKnownState(double Julia
 	
 	FVector BodyLocation_Heliocentric_FK5J2000_AU = VSOPEclipticToFK5Equatorial.TransformVector(BodyLocation_VSOP87);
 
-	// Compute the Observer Location in the same, applying the effects of Precession matrix. 
-	FVector ObserverLocationGeoCentric_FK5J2000AU = GetObserverGeocentricLocationAU(ObserverLatitude, ObserverLongitude, 0, JulianDate );
-	ObserverLocationGeoCentric_FK5J2000AU = NutationMatrix.TransformPosition(ObserverLocationGeoCentric_FK5J2000AU);
-	ObserverLocationGeoCentric_FK5J2000AU = PrecessionMatrix.TransformPosition(ObserverLocationGeoCentric_FK5J2000AU); 
+	// Compute the Observer Location in the same, applying the effects of Precession matrix.
+	FVector ObserverLocationGeoCentric_FK5J2000AU = FVector::ZeroVector;
+	if (!bGeoCentric)
+	{
+		ObserverLocationGeoCentric_FK5J2000AU = GetObserverGeocentricLocationAU(ObserverLatitude, ObserverLongitude, 0, JulianDate );
+		ObserverLocationGeoCentric_FK5J2000AU = NutationMatrix.TransformPosition(ObserverLocationGeoCentric_FK5J2000AU);
+		ObserverLocationGeoCentric_FK5J2000AU = PrecessionMatrix.TransformPosition(ObserverLocationGeoCentric_FK5J2000AU);
+	}
+
 
 	// In the FK5J2000AU Frame (S = Sun, B = Body, E = Earth, O = Observer
 	//   B
@@ -277,9 +282,24 @@ double UCelestialMaths::GetPlanetaryBodyMagnitude(EVSOP87BodyType VSOP87BodyType
 
 	
 	
-	double DistanceFactor = 5.0 * FMath::LogX(10.0, DistanceToEarthAU * DistanceToSunAU);
-	double ApparentMagnitude = 0.0; 
-	PhaseAngle = FMath::RadiansToDegrees(FMath::Acos( (DistanceToSunAU * DistanceToSunAU + DistanceToEarthAU * DistanceToEarthAU - DistanceEarthToSunAU*DistanceEarthToSunAU) / (2.0 * DistanceToSunAU * DistanceToEarthAU) ));
+	double DistanceFactor = 5.0 * FMath::LogX(10.0, FMath::Abs(DistanceToEarthAU * DistanceToSunAU));
+	if (VSOP87BodyType == EVSOP87BodyType::Moon)
+	{
+		// For the moon, the distance factor is different. // Schaefer 1998
+		// mV = −12.73 + 0.026∣α∣ + 4.0×10−9α4 + 5log10(Δ0/Δ)
+		// Δ0 =384400 km (average distance), Δ = Moon-Earth distance
+		DistanceFactor = 5.0 * FMath::LogX(10.0, FMath::Abs(DistanceToEarthAU) / MetersToAstronomicalUnits(384400.0 * UE_KM_TO_M));
+	}
+	
+	double ApparentMagnitude = 0.0;
+	if (DistanceToSunAU <= 0 || DistanceToEarthAU <= 0)
+	{
+		PhaseAngle = 0.0;
+		return 0.0;
+	}
+	
+	double PhaseInput = (DistanceToSunAU * DistanceToSunAU + DistanceToEarthAU * DistanceToEarthAU - DistanceEarthToSunAU * DistanceEarthToSunAU) / (2.0 * DistanceToSunAU * DistanceToEarthAU); 
+	PhaseAngle = FMath::RadiansToDegrees(FMath::Acos( FMath::Clamp(PhaseInput , -1.0, 1.0)));
 	double PhaseAngle2 = PhaseAngle * PhaseAngle;
 	double PhaseAngle3 = PhaseAngle2 * PhaseAngle;
 	double PhaseAngle4 = PhaseAngle3 * PhaseAngle;
@@ -857,10 +877,16 @@ FTransform UCelestialMaths::GetEarthCenterTransformECEF(double Latitude, double
 	return FTransform(AxisMatrix.Inverse());
 }
 
-FTransform UCelestialMaths::GetEarthCenterTransformUEFrame(double Latitude, double Longitude, double Altitude)
+FTransform UCelestialMaths::GetEarthCenterTransformUEFrame(double Latitude, double Longitude, double Altitude, bool Geocentric)
 {
 	// Get the transformation in proper Celestial World units (meters, right handed)
-	FTransform ECEFFrameToWorldFrame = GetEarthCenterTransformECEF(Latitude, Longitude, 0);
+	FTransform ECEFFrameToWorldFrame = FTransform(); // By default, in geocentric mode, Center is a Zero Transform
+	if (!Geocentric)
+	{
+		// Consider tangent frame at lat/long
+		ECEFFrameToWorldFrame = GetEarthCenterTransformECEF(Latitude, Longitude, 0);
+	}
+	
 
 	// UE Frame are expressed in Left-handed coordinates, and units are in meters - Convert to UE Transform 
 	FMatrix WorldFrameToUEFrame = FMatrix( 
@@ -982,17 +1008,17 @@ FVector UCelestialMaths::ComputeAberration(const FVector& TopocentricTargetLocat
 
 #pragma region Sun & Moon
 
-FStellarBody UCelestialMaths::GetSunInformation(double JulianDate, double ObserverLatitude, double ObserverLongitude)
+FStellarBody UCelestialMaths::GetSunInformation(double JulianDate, double ObserverLatitude, double ObserverLongitude, bool bGeoCentric)
 {
 	TRACE_CPUPROFILER_EVENT_SCOPE(GetSunInformation);
 	
 	// Compute Earth Reference location
 	FPlanetaryBodyKinematicState EarthKinematicState = GetPlanetaryBodyKinematicState_AU(JulianDate, EVSOP87BodyType::Earth);
 
-	return GetSunInformation_UsingKnownState(JulianDate, EarthKinematicState, ObserverLatitude, ObserverLongitude);
+	return GetSunInformation_UsingKnownState(JulianDate, EarthKinematicState, ObserverLatitude, ObserverLongitude, bGeoCentric);
 }
 
-FStellarBody UCelestialMaths::GetSunInformation_UsingKnownState(double JulianDate, const FPlanetaryBodyKinematicState& EarthKinematicState, double ObserverLatitude, double ObserverLongitude)
+FStellarBody UCelestialMaths::GetSunInformation_UsingKnownState(double JulianDate, const FPlanetaryBodyKinematicState& EarthKinematicState, double ObserverLatitude, double ObserverLongitude, bool bGeoCentric)
 {
 	TRACE_CPUPROFILER_EVENT_SCOPE(GetSunInformation_UsingKnownState);
 	// Get the Celestial Body Coordinates
@@ -1003,7 +1029,7 @@ FStellarBody UCelestialMaths::GetSunInformation_UsingKnownState(double JulianDat
 	double DistanceBodyToEarthAU;
 	double DistanceBodyToSunAU;
 	double DistanceEarthToSunAU;
-	GetBodyCelestialCoordinatesAU_UsingKnownState(JulianDate, EVSOP87BodyType::Sun, EarthKinematicState, ObserverLatitude, ObserverLongitude, false, RAJ2000Hours, DECJ2000Degrees, RAHours, DECDegrees, DistanceBodyToEarthAU, DistanceBodyToSunAU, DistanceEarthToSunAU);  
+	GetBodyCelestialCoordinatesAU_UsingKnownState(JulianDate, EVSOP87BodyType::Sun, EarthKinematicState, ObserverLatitude, ObserverLongitude, bGeoCentric, false, RAJ2000Hours, DECJ2000Degrees, RAHours, DECDegrees, DistanceBodyToEarthAU, DistanceBodyToSunAU, DistanceEarthToSunAU);  
 
 	// Set the Sun properties. Some properties can be hardcoded for the Sun (CI, Magnitude, Radius) 
 	FStellarBody SunInfo;
diff --git a/Engine/Plugins/Experimental/CelestialVault/Source/CelestialVault/Private/CelestialVaultDaySequenceActor.cpp b/Engine/Plugins/Experimental/CelestialVault/Source/CelestialVault/Private/CelestialVaultDaySequenceActor.cpp
--- a/Engine/Plugins/Experimental/CelestialVault/Source/CelestialVault/Private/CelestialVaultDaySequenceActor.cpp
+++ b/Engine/Plugins/Experimental/CelestialVault/Source/CelestialVault/Private/CelestialVaultDaySequenceActor.cpp
@@ -18,6 +18,7 @@
 #include "UObject/ConstructorHelpers.h"
 #include "DaySequenceSubsystem.h"
 #include "DaySequenceCollectionAsset.h"
+#include "DaySequence/Private/DaySequencePlayer.h"
 #include "Engine/World.h"
 #include "Engine/StaticMesh.h"
 #include "Curves/CurveFloat.h"
@@ -50,66 +51,102 @@ ACelestialVaultDaySequenceActor::ACelestialVaultDaySequenceActor(const FObjectIn
 : Super(Init)
 {
 	ExponentialHeightFogComponent = CreateOptionalDefaultSubobject<UExponentialHeightFogComponent>(TEXT("ExponentialHeightFog"));
-	ExponentialHeightFogComponent->SetupAttachment(RootComponent);
-
+	if (ExponentialHeightFogComponent)
+	{
+		ExponentialHeightFogComponent->SetupAttachment(RootComponent);
+		ExponentialHeightFogComponent->FogCutoffDistance = 20000000.0; // We don't want the fog to apply on the deep sky
+		ExponentialHeightFogComponent->EndDistance = 500000.0; // We need a non-zero value here to avoid artifacts in round worlds, and quite high to have a proper haze
+	}
+	
 	static ConstructorHelpers::FObjectFinder<UCurveFloat> HighlightContrastCurve(TEXT("/CelestialVault/Data/CF_CelestialHighlightContrastCurve.CF_CelestialHighlightContrastCurve"));
 	GlobalPostProcessVolume = CreateOptionalDefaultSubobject<UPostProcessComponent>(TEXT("GlobalPostProcessVolume"));
-	GlobalPostProcessVolume->SetupAttachment(RootComponent);
-	GlobalPostProcessVolume->Settings.bOverride_AutoExposureMinBrightness = true;
-	GlobalPostProcessVolume->Settings.AutoExposureMinBrightness = -0.5f;
-	GlobalPostProcessVolume->Settings.bOverride_LocalExposureHighlightContrastCurve = true;
-	GlobalPostProcessVolume->Settings.LocalExposureHighlightContrastCurve = HighlightContrastCurve.Object;
-	GlobalPostProcessVolume->Settings.bOverride_LocalExposureDetailStrength = true;
-	GlobalPostProcessVolume->Settings.LocalExposureDetailStrength = 1.2f;
-	GlobalPostProcessVolume->Settings.bOverride_LocalExposureBlurredLuminanceBlend = true;
-	GlobalPostProcessVolume->Settings.LocalExposureBlurredLuminanceBlend = 0.4f;
+	if (GlobalPostProcessVolume)
+	{
+		GlobalPostProcessVolume->SetupAttachment(RootComponent);
+		GlobalPostProcessVolume->Settings.bOverride_AutoExposureMinBrightness = true;
+		GlobalPostProcessVolume->Settings.AutoExposureMinBrightness = -0.5f;
+		GlobalPostProcessVolume->Settings.bOverride_LocalExposureHighlightContrastCurve = true;
+		GlobalPostProcessVolume->Settings.LocalExposureHighlightContrastCurve = HighlightContrastCurve.Object;
+		GlobalPostProcessVolume->Settings.bOverride_LocalExposureDetailStrength = true;
+		GlobalPostProcessVolume->Settings.LocalExposureDetailStrength = 1.2f;
+		GlobalPostProcessVolume->Settings.bOverride_LocalExposureBlurredLuminanceBlend = true;
+		GlobalPostProcessVolume->Settings.LocalExposureBlurredLuminanceBlend = 0.4f;
+	}
 	
 	// Components Attached to Root
 	SkyLightComponent = CreateDefaultSubobject<USkyLightComponent>(TEXT("Sky Light"));
-	SkyLightComponent->SetupAttachment(RootComponent);
-	SkyLightComponent->bRealTimeCapture = true;
-	SkyLightComponent->bLowerHemisphereIsBlack = false;
+	if (SkyLightComponent)
+	{
+		SkyLightComponent->SetupAttachment(RootComponent);
+		SkyLightComponent->bRealTimeCapture = true;
+		SkyLightComponent->bLowerHemisphereIsBlack = false;
+	}
+	
 
 	VolumetricCloudComponent = CreateOptionalDefaultSubobject<UVolumetricCloudComponent>(TEXT("Volumetric Cloud"));
-	VolumetricCloudComponent->SetupAttachment(RootComponent);
-	if (!IsTemplate())
+	if (VolumetricCloudComponent)
 	{
-		// We don't want to load this material for the CDO as it will hold on to it forever and it is quite a large asset.
-		static ConstructorHelpers::FObjectFinder<UMaterialInterface> VolumetricCloudDefaultMaterialRef(TEXT("/Engine/EngineSky/VolumetricClouds/m_SimpleVolumetricCloud_Inst.m_SimpleVolumetricCloud_Inst"));
-		VolumetricCloudComponent->SetMaterial(VolumetricCloudDefaultMaterialRef.Object);
+		VolumetricCloudComponent->SetupAttachment(RootComponent);
+		if (!IsTemplate())
+		{
+			// We don't want to load this material for the CDO as it will hold on to it forever and it is quite a large asset.
+			static ConstructorHelpers::FObjectFinder<UMaterialInterface> VolumetricCloudDefaultMaterialRef(TEXT("/Engine/EngineSky/VolumetricClouds/m_SimpleVolumetricCloud_Inst.m_SimpleVolumetricCloud_Inst"));
+			VolumetricCloudComponent->SetMaterial(VolumetricCloudDefaultMaterialRef.Object);
+		}	
 	}
+	
 
 	NorthOffsetComponent = CreateDefaultSubobject<USceneComponent>(TEXT("North Offset Transform"));
-	NorthOffsetComponent->SetupAttachment(RootComponent);
+	if (NorthOffsetComponent)
+	{
+		NorthOffsetComponent->SetupAttachment(RootComponent);	
+	}
+	
 
 	// Add an editor-only compass to visualize the North Offset
 	CompassComponent = CreateOptionalDefaultSubobject<UStaticMeshComponent>(TEXT("Compass"));
-	CompassComponent->SetupAttachment(NorthOffsetComponent);
-	CompassComponent->SetCollisionEnabled(ECollisionEnabled::NoCollision);
-	CompassComponent->SetGenerateOverlapEvents(false);
-	CompassComponent->SetCastShadow(false);
-	CompassComponent->SetAffectDynamicIndirectLighting(false);
-	CompassComponent->SetCanEverAffectNavigation(false);
-	CompassComponent->bIsEditorOnly = true;
-	CompassComponent->SetHiddenInGame(true);
-	if (!IsTemplate())
+	if (CompassComponent)
 	{
-		static ConstructorHelpers::FObjectFinder<UStaticMesh> CompassDefaultMesh(TEXT("/CelestialVault/Editor/SM_Compass.SM_Compass"));
-		CompassComponent->SetStaticMesh(CompassDefaultMesh.Object);
+		CompassComponent->SetupAttachment(NorthOffsetComponent);
+		CompassComponent->SetCollisionEnabled(ECollisionEnabled::NoCollision);
+		CompassComponent->SetGenerateOverlapEvents(false);
+		CompassComponent->SetCastShadow(false);
+		CompassComponent->SetAffectDynamicIndirectLighting(false);
+		CompassComponent->SetCanEverAffectNavigation(false);
+		CompassComponent->bIsEditorOnly = true;
+		CompassComponent->SetHiddenInGame(true);
+		if (!IsTemplate())
+		{
+			static ConstructorHelpers::FObjectFinder<UStaticMesh> CompassDefaultMesh(TEXT("/CelestialVault/Editor/SM_Compass.SM_Compass"));
+			CompassComponent->SetStaticMesh(CompassDefaultMesh.Object);
+		}	
 	}
 	
+	
 	PlanetCenterComponent = CreateDefaultSubobject<USceneComponent>(TEXT("Planet Center Transform"));
-	PlanetCenterComponent->SetupAttachment(NorthOffsetComponent);
+	if (PlanetCenterComponent)
+	{
+		PlanetCenterComponent->SetupAttachment(NorthOffsetComponent);	
+	}
+	
 	
 	// Components Attached to Planet Center
 	SkyAtmosphereComponent = CreateDefaultSubobject<USkyAtmosphereComponent>(TEXT("Sky Atmosphere"));
-	SkyAtmosphereComponent->SetupAttachment(PlanetCenterComponent);
-	SkyAtmosphereComponent->TransformMode = ESkyAtmosphereTransformMode::PlanetCenterAtComponentTransform;
+	if (SkyAtmosphereComponent)
+	{
+		SkyAtmosphereComponent->SetupAttachment(PlanetCenterComponent);
+		SkyAtmosphereComponent->TransformMode = ESkyAtmosphereTransformMode::PlanetCenterAtComponentTransform;	
+	}
+	
 
 
 	// Rotating Celestial Vault 
 	CelestialVaultComponent = CreateDefaultSubobject<USceneComponent>(TEXT("Rotating Celestial Vault"));
-	CelestialVaultComponent->SetupAttachment(PlanetCenterComponent);
+	if (CelestialVaultComponent)
+	{
+		CelestialVaultComponent->SetupAttachment(PlanetCenterComponent);	
+	}
+	
 
 	// Components attached to the Celestial Vaul
 	
@@ -117,74 +154,95 @@ ACelestialVaultDaySequenceActor::ACelestialVaultDaySequenceActor(const FObjectIn
 	static ConstructorHelpers::FObjectFinder<UStaticMesh> SkySphereDefaultMesh(TEXT("/CelestialVault/Meshes/SM_CelestialVault.SM_CelestialVault"));
 	static ConstructorHelpers::FObjectFinder<UMaterialInterface> SkySphereDefaultMaterial(TEXT("/CelestialVault/Materials/MI_CelestialVault.MI_CelestialVault"));
 	DeepSkyComponent = CreateOptionalDefaultSubobject<UStaticMeshComponent>(TEXT("Deep Sky"));
-	DeepSkyComponent->SetupAttachment(CelestialVaultComponent);
-	DeepSkyComponent->SetStaticMesh(SkySphereDefaultMesh.Object);
-	DeepSkyComponent->SetCollisionEnabled(ECollisionEnabled::NoCollision);
-	DeepSkyComponent->SetGenerateOverlapEvents(false);
-	DeepSkyComponent->SetCastShadow(false);
-	DeepSkyComponent->SetAffectDynamicIndirectLighting(false);
-	DeepSkyComponent->SetCanEverAffectNavigation(false);
-	DeepSkyComponent->SetMaterial(0, SkySphereDefaultMaterial.Object.Get());
-	DeepSkyComponent->SetRelativeScale3D(FVector(CelestialVaultDistance*1000.0));
+	if (DeepSkyComponent)
+	{
+		DeepSkyComponent->SetupAttachment(CelestialVaultComponent);
+		DeepSkyComponent->SetStaticMesh(SkySphereDefaultMesh.Object);
+		DeepSkyComponent->SetCollisionEnabled(ECollisionEnabled::NoCollision);
+		DeepSkyComponent->SetGenerateOverlapEvents(false);
+		DeepSkyComponent->SetCastShadow(false);
+		DeepSkyComponent->SetAffectDynamicIndirectLighting(false);
+		DeepSkyComponent->SetCanEverAffectNavigation(false);
+		DeepSkyComponent->SetMaterial(0, SkySphereDefaultMaterial.Object.Get());
+		DeepSkyComponent->SetRelativeScale3D(FVector(CelestialVaultDistance*1000.0));	
+	}
+	
 
 	// Stars ISM
 	static ConstructorHelpers::FObjectFinder<UStaticMesh> PlaneXMesh(TEXT("/CelestialVault/Meshes/SM_Plane_FacingX.SM_Plane_FacingX"));
 	static ConstructorHelpers::FObjectFinder<UMaterialInterface> StarsDefaultMaterial(TEXT("/CelestialVault/Materials/MI_Stars.MI_Stars"));
 	StarsComponent = CreateOptionalDefaultSubobject<UInstancedStaticMeshComponent>(TEXT("Stars"));
-	StarsComponent->SetupAttachment(DeepSkyComponent);
-	StarsComponent->SetStaticMesh(PlaneXMesh.Object);
-	StarsComponent->SetCollisionEnabled(ECollisionEnabled::NoCollision);
-	StarsComponent->SetGenerateOverlapEvents(false);
-	StarsComponent->SetCastShadow(false);
-	StarsComponent->SetAffectDynamicIndirectLighting(false);
-	StarsComponent->SetCanEverAffectNavigation(false);
-	StarsComponent->SetMaterial(0, StarsDefaultMaterial.Object);
-	StarsComponent->bUseAttachParentBound = true;
+	if (StarsComponent)
+	{
+		StarsComponent->SetupAttachment(DeepSkyComponent);
+		StarsComponent->SetStaticMesh(PlaneXMesh.Object);
+		StarsComponent->SetCollisionEnabled(ECollisionEnabled::NoCollision);
+		StarsComponent->SetGenerateOverlapEvents(false);
+		StarsComponent->SetCastShadow(false);
+		StarsComponent->SetAffectDynamicIndirectLighting(false);
+		StarsComponent->SetCanEverAffectNavigation(false);
+		StarsComponent->SetMaterial(0, StarsDefaultMaterial.Object);
+		StarsComponent->bUseAttachParentBound = true;	
+	}
+	
 
 	// Planets ISM
 	static ConstructorHelpers::FObjectFinder<UMaterialInterface> PlanetsDefaultMaterial(TEXT("/CelestialVault/Materials/MI_SolarSystemPlanets.MI_SolarSystemPlanets"));
 	PlanetsComponent = CreateOptionalDefaultSubobject<UInstancedStaticMeshComponent>(TEXT("Planets"));
-	PlanetsComponent->SetupAttachment(CelestialVaultComponent);
-	PlanetsComponent->SetStaticMesh(PlaneXMesh.Object);
-	PlanetsComponent->SetCollisionEnabled(ECollisionEnabled::NoCollision);
-	PlanetsComponent->SetGenerateOverlapEvents(false);
-	PlanetsComponent->SetCastShadow(false);
-	PlanetsComponent->SetAffectDynamicIndirectLighting(false);
-	PlanetsComponent->SetCanEverAffectNavigation(false);
-	PlanetsComponent->SetMaterial(0, PlanetsDefaultMaterial.Object);
+	if (PlanetsComponent)
+	{
+		PlanetsComponent->SetupAttachment(CelestialVaultComponent);
+		PlanetsComponent->SetStaticMesh(PlaneXMesh.Object);
+		PlanetsComponent->SetCollisionEnabled(ECollisionEnabled::NoCollision);
+		PlanetsComponent->SetGenerateOverlapEvents(false);
+		PlanetsComponent->SetCastShadow(false);
+		PlanetsComponent->SetAffectDynamicIndirectLighting(false);
+		PlanetsComponent->SetCanEverAffectNavigation(false);
+		PlanetsComponent->SetMaterial(0, PlanetsDefaultMaterial.Object);
+	}
+	
 	
 
 	// Moon - Disc
 	static ConstructorHelpers::FObjectFinder<UMaterialInterface> MoonDiscDefaultMaterial(TEXT("/CelestialVault/Materials/MI_Moon.MI_Moon"));
 	MoonDiscComponent = CreateOptionalDefaultSubobject<UStaticMeshComponent>(TEXT("Moon Disk"));
-	MoonDiscComponent->SetupAttachment(CelestialVaultComponent);
-	MoonDiscComponent->SetStaticMesh(PlaneXMesh.Object);
-	MoonDiscComponent->SetCollisionEnabled(ECollisionEnabled::NoCollision);
-	MoonDiscComponent->SetGenerateOverlapEvents(false);
-	MoonDiscComponent->SetCastShadow(true); // Eclipses?
-	MoonDiscComponent->SetAffectDynamicIndirectLighting(false);
-	MoonDiscComponent->SetCanEverAffectNavigation(false);
-	MoonDiscComponent->SetMaterial(0, MoonDiscDefaultMaterial.Object.Get());
+	if (MoonDiscComponent)
+	{
+		MoonDiscComponent->SetupAttachment(CelestialVaultComponent);
+		MoonDiscComponent->SetStaticMesh(PlaneXMesh.Object);
+		MoonDiscComponent->SetCollisionEnabled(ECollisionEnabled::NoCollision);
+		MoonDiscComponent->SetGenerateOverlapEvents(false);
+		MoonDiscComponent->SetCastShadow(true); // Eclipses?
+		MoonDiscComponent->SetAffectDynamicIndirectLighting(false);
+		MoonDiscComponent->SetCanEverAffectNavigation(false);
+		MoonDiscComponent->SetMaterial(0, MoonDiscDefaultMaterial.Object.Get());
+	}
 	
 
 	// Moon - Light
 	MoonLightComponent = CreateDefaultSubobject<UDirectionalLightComponent>(TEXT("Moon Light"));
-	MoonLightComponent->SetupAttachment(CelestialVaultComponent);
-	MoonLightComponent->SetAtmosphereSunLightIndex(1);	// Make Moon the secondary directional light that contributes to the sky atmosphere.
-	MoonLightComponent->SetForwardShadingPriority(0); // Give Moon forward shading priority.
-	MoonLightComponent->SetIntensity(MoonLightIntensity);
-	MoonLightComponent->SetUseTemperature(true);
-	MoonLightComponent->SetTemperature(9000.f);
-	MoonLightComponent->SetWorldRotation(FRotator(-45.f, 0.0f, 0.0f));
-	MoonLightComponent->bCastCloudShadows = true; // Otherwise we still have hard shadows with an overcast sky 
+	if (MoonLightComponent)
+	{
+		MoonLightComponent->SetupAttachment(CelestialVaultComponent);
+		MoonLightComponent->SetAtmosphereSunLightIndex(1);	// Make Moon the secondary directional light that contributes to the sky atmosphere.
+		MoonLightComponent->SetForwardShadingPriority(0); // Give Moon forward shading priority.
+		MoonLightComponent->SetIntensity(MoonLightIntensity);
+		MoonLightComponent->SetUseTemperature(true);
+		MoonLightComponent->SetTemperature(9000.f);
+		MoonLightComponent->SetWorldRotation(FRotator(-45.f, 0.0f, 0.0f));
+		MoonLightComponent->bCastCloudShadows = true; // Otherwise we still have hard shadows with an overcast sky 
+	}
 	
 	// Attach the sunlight relative to the rotating Vault
 	SunLightComponent = CreateDefaultSubobject<UDirectionalLightComponent>(TEXT("Sun Light"));
-	SunLightComponent->SetupAttachment(CelestialVaultComponent);
-	SunLightComponent->SetAtmosphereSunLightIndex(0);	// Make Sun the first directional light that contributes to the sky atmosphere.
-	SunLightComponent->SetForwardShadingPriority(1);   // Give Sun forward shading priority.
-	SunLightComponent->SetIntensity(SunLightIntensity);
-	SunLightComponent->bCastCloudShadows = true; // Otherwise we still have hard shadows with an overcast sky
+	if (SunLightComponent)
+	{
+		SunLightComponent->SetupAttachment(CelestialVaultComponent);
+		SunLightComponent->SetAtmosphereSunLightIndex(0);	// Make Sun the first directional light that contributes to the sky atmosphere.
+		SunLightComponent->SetForwardShadingPriority(1);   // Give Sun forward shading priority.
+		SunLightComponent->SetIntensity(SunLightIntensity);
+		SunLightComponent->bCastCloudShadows = true; // Otherwise we still have hard shadows with an overcast sky
+	}
 	
 	// Sequence and Data Assets
 	if (!IsTemplate())
@@ -322,6 +380,19 @@ void ACelestialVaultDaySequenceActor::UnregisterAllComponents(bool bForReregiste
 	Super::UnregisterAllComponents(bForReregister);
 }
 
+void ACelestialVaultDaySequenceActor::Stop()
+{
+	UDaySequencePlayer* Player = GetSequencePlayerInternal();
+	UWorld* World = GetWorld();
+	if (HasValidRootSequence() && Player && World && World->IsGameWorld())
+	{
+		Player->Stop();
+		StopDaySequenceUpdateTimer();
+		SetTimeOfDay( GetInitialTimeOfDay());
+		SequenceLoopCount = 0;
+	}
+}
+
 #if CELESTIAL_VAULT_ENABLE_DRAW_DEBUG
 
 void ACelestialVaultDaySequenceActor::DebugDrawCallback(UCanvas* Canvas, APlayerController* PC) const
@@ -389,7 +460,7 @@ FDateTime ACelestialVaultDaySequenceActor::GetDate() const
 		return FDateTime(Year, ClampedMonth, ClampedDay, 0,0,0);
 	}
 		
-	return FDateTime(Year, static_cast<int>(Month), Day, 0,0,0) + FTimespan(SequenceLoopCount, 0,0, 0);;
+	return FDateTime(Year, ClampedMonth, Day, 0,0,0) + FTimespan(SequenceLoopCount, 0,0, 0);;
 }
 
 FDateTime ACelestialVaultDaySequenceActor::GetDateAndTime() const
@@ -422,7 +493,7 @@ FStellarBody ACelestialVaultDaySequenceActor::ComputeSunInfo_UsingKnownState(dou
 	TRACE_CPUPROFILER_EVENT_SCOPE(ComputeSunInfo_UsingKnownState);
 
 	
-	FStellarBody ResultSunInfo = UCelestialMaths::GetSunInformation_UsingKnownState(JulianDate, EarthKinematicState, Latitude, Longitude);
+	FStellarBody ResultSunInfo = UCelestialMaths::GetSunInformation_UsingKnownState(JulianDate, EarthKinematicState, Latitude, Longitude, bGeocentric);
 
 	// Location
 	FVector SunLocation = UCelestialMaths::RADECToXYZ_RH(ResultSunInfo.RA * 15.0, ResultSunInfo.DEC, UCelestialMaths::AstronomicalUnitsToMeters(ResultSunInfo.DistanceInAU) * 100); 
@@ -986,7 +1057,7 @@ FPlanetaryBody ACelestialVaultDaySequenceActor::GetPlanetaryBodyInfo_UsingKnownS
 	// Compute location	
 	double RAJ2000Hours, DECJ2000Degrees, RAHours, DECDegrees, DistanceToEarthAU, DistanceToSunAU, DistanceEarthToSunAU;
 	
-	UCelestialMaths::GetBodyCelestialCoordinatesAU_UsingKnownState(JulianDay, InputPlanetaryBody.VSOP87BodyType, EarthKinematicState, Latitude, Longitude, bIgnoreRelativisticEffect, RAJ2000Hours, DECJ2000Degrees,RAHours, DECDegrees, DistanceToEarthAU, DistanceToSunAU, DistanceEarthToSunAU);
+	UCelestialMaths::GetBodyCelestialCoordinatesAU_UsingKnownState(JulianDay, InputPlanetaryBody.VSOP87BodyType, EarthKinematicState, Latitude, Longitude, bGeocentric, bIgnoreRelativisticEffect, RAJ2000Hours,DECJ2000Degrees, RAHours, DECDegrees, DistanceToEarthAU, DistanceToSunAU, DistanceEarthToSunAU);
 	BodyInfo.RAJ2000 = RAJ2000Hours;
 	BodyInfo.DECJ2000 = DECJ2000Degrees;
 	BodyInfo.RA = RAHours;
@@ -995,7 +1066,7 @@ FPlanetaryBody ACelestialVaultDaySequenceActor::GetPlanetaryBodyInfo_UsingKnownS
 
 	// Compute Magnitude
 	double Phase; // TODO - Check the PhaseAngle Computations. 
-	BodyInfo.Magnitude = UCelestialMaths::GetPlanetaryBodyMagnitude(InputPlanetaryBody.VSOP87BodyType, DistanceToEarthAU, DistanceToSunAU, DistanceEarthToSunAU, Phase);
+	BodyInfo.Magnitude = UCelestialMaths::GetPlanetaryBodyMagnitude(InputPlanetaryBody.VSOP87BodyType, DistanceToSunAU, DistanceToEarthAU , DistanceEarthToSunAU, Phase);
 	BodyInfo.Age = Phase;
 	
 	// Compute the True and the Scaled apparent Diameters
@@ -1006,12 +1077,20 @@ FPlanetaryBody ACelestialVaultDaySequenceActor::GetPlanetaryBodyInfo_UsingKnownS
 void ACelestialVaultDaySequenceActor::UpdatePlanetCenter()
 {
 	// Get the transformation in proper UE Frame units, and set it to the Component 
-	PlanetCenterTransform = UCelestialMaths::GetEarthCenterTransformUEFrame(Latitude, Longitude, 0);
-	PlanetCenterComponent->SetWorldTransform(PlanetCenterTransform);
+	PlanetCenterTransform = UCelestialMaths::GetEarthCenterTransformUEFrame(Latitude, Longitude, 0, bGeocentric);
+	PlanetCenterComponent->SetRelativeTransform(PlanetCenterTransform);
 
 	// Adjust the SkyAtmosphere to the new radius
-	double Radius = PlanetCenterTransform.GetLocation().Length();
-	SkyAtmosphereComponent->SetBottomRadius(Radius / UE_ONE_KILOMETER);
+
+	if (bGeocentric)
+	{
+		SkyAtmosphereComponent->SetBottomRadius(6378.137); // Reset to max Earth ellipsoidal radius. 
+	}
+	else
+	{
+		double Radius = PlanetCenterTransform.GetLocation().Length();
+		SkyAtmosphereComponent->SetBottomRadius(Radius / UE_ONE_KILOMETER);
+	}
 }
 
 void ACelestialVaultDaySequenceActor::RebuildAll()
@@ -1034,6 +1113,7 @@ void ACelestialVaultDaySequenceActor::RebuildAll()
 		InitStarsComponent();
 		InitPlanetsComponent();
 		ResetBodyMotionCache();
+		UpdateBodiesMotion();
 	}
 }
 
@@ -1061,10 +1141,11 @@ void ACelestialVaultDaySequenceActor::UpdateBodiesMotion()
 	if ( SunLightComponent)
 	{
 		SunLightComponent->SetRelativeRotation(SunLightRotationCurve.Eval(TimeOfDay));
+		SunLightComponent->SetIntensity(SunLightIntensity);
 	}
 
 	// Don't use a curve cache for the Celestial Vault rotation. It's a basic angular rotation. 
-	double VaultYaw = GetCelestialVaultAngle(SequenceLoopCount * 24 + TimeOfDay );
+	double VaultYaw = GetCelestialVaultAngle(TimeOfDay );
 	CelestialVaultComponent->SetRelativeRotation(FRotator(0.0, VaultYaw, 0.0));
 }
 
@@ -1123,7 +1204,6 @@ void ACelestialVaultDaySequenceActor::AddBodyMotionCacheKey(double TimeOfDay)
 	// We need the Julian Date for the celestial computation. For this, we need to consider the initial day, the number of times the sequence has looped, and the current time of day.
 	// Compute the Local time
 	FDateTime LocalTime = GetDate() + FTimespan::FromHours(TimeOfDay);
-	
 	// Convert DateTime to Julian Date 
 	FDateTime UTCTime = UCelestialMaths::LocalTimeToUTCTime(LocalTime, GMT_TimeZone, IsDaylightSavingsAtTimeOfDay(TimeOfDay));
 	double JulianDate = UCelestialMaths::UTCDateTimeToJulianDate(UTCTime);
diff --git a/Engine/Plugins/Experimental/CelestialVault/Source/CelestialVault/Public/CelestialMaths.h b/Engine/Plugins/Experimental/CelestialVault/Source/CelestialVault/Public/CelestialMaths.h
--- a/Engine/Plugins/Experimental/CelestialVault/Source/CelestialVault/Public/CelestialMaths.h
+++ b/Engine/Plugins/Experimental/CelestialVault/Source/CelestialVault/Public/CelestialMaths.h
@@ -97,11 +97,12 @@ public:
      *
 	 * By default, the relativistic effects are considered (time taken for the light to reach the Observer body location)
 	 * This function therefore returns the location of the Planetary Body as if it was seen from the Observer Body Location at this JulianDate
-	 * With bIgnoreRelativisticEffect=true, one can ignore this and have the instant location. It can also save some computation cycles for close bodies like Moon. 
+	 * With bIgnoreRelativisticEffect=true, one can ignore this and have the instant location. It can also save some computation cycles for close bodies like Moon.
+	 * With bGeoCentric=true, the observer is expected to be at the Earth center, and the Latitude/longitude are ignored. 
      * **/
     UFUNCTION(BlueprintCallable, Category="Celestial Vault|Planetary bodies")
-    static void GetBodyCelestialCoordinatesAU(double JulianDate, EVSOP87BodyType VSOP87BodyType, double ObserverLatitude, double ObserverLongitude, bool bIgnoreRelativisticEffect,
-    	double& RAJ2000Hours, double& DECJ2000Degrees, double& RAHours, double& DECDegrees, double& DistanceBodyToEarthAU, double& DistanceBodyToSunAU, double& DistanceEarthToSunAU);
+    static void GetBodyCelestialCoordinatesAU(double JulianDate, EVSOP87BodyType VSOP87BodyType, double ObserverLatitude, double ObserverLongitude, bool bGeoCentric,
+                                              bool bIgnoreRelativisticEffect, double& RAJ2000Hours, double& DECJ2000Degrees, double& RAHours, double& DECDegrees, double& DistanceBodyToEarthAU, double& DistanceBodyToSunAU, double& DistanceEarthToSunAU);
 
 	/** Return the location of a Planetary Body relative to the Earth, expressed in Celestial Coordinates (RA, DEC, Distance)
 	 * It requires the Observer location on Earth for more precise computations
@@ -114,8 +115,8 @@ public:
 	 * This override can run faster if we provide it with a cache of the Earth Location and Velocity, because it won't recompute it. 
 	 * **/
 	UFUNCTION(BlueprintCallable, Category="Celestial Vault|Planetary bodies")
-	static void GetBodyCelestialCoordinatesAU_UsingKnownState(double JulianDate, EVSOP87BodyType VSOP87BodyType, const FPlanetaryBodyKinematicState& EarthKinematicState, double ObserverLatitude, double ObserverLongitude, bool bIgnoreRelativisticEffect,
-		double& RAJ2000Hours, double& DECJ2000Degrees, double& RAHours, double& DECDegrees, double& DistanceBodyToEarthAU, double& DistanceBodyToSunAU, double& DistanceEarthToSunAU);
+	static void GetBodyCelestialCoordinatesAU_UsingKnownState(double JulianDate, EVSOP87BodyType VSOP87BodyType, const FPlanetaryBodyKinematicState& EarthKinematicState, double ObserverLatitude, double ObserverLongitude, bool bGeoCentric,
+	                                                          bool bIgnoreRelativisticEffect, double& RAJ2000Hours, double& DECJ2000Degrees, double& RAHours, double& DECDegrees, double& DistanceBodyToEarthAU, double& DistanceBodyToSunAU, double& DistanceEarthToSunAU);
 	
     /** Return the Magnitude of a Planetary Body as seen from the Earth **/
     UFUNCTION(BlueprintCallable, Category="Celestial Vault|Planetary bodies")
@@ -309,9 +310,11 @@ public:
 	/** Return the Transformation to apply to a WGS84 Ellipsoid model so that its location in Lat,long,Altitude is tangent to the Origin
 	 * It's used to locate the Rotating Celestial Vault for a specific UE Origin
 	 * The Transform is expressed in the UE Frame (Location units are UE Units, left-handed Unreal Frame)
+	 *
+	 * With Geocentric=true, the observer is expected to be at the Earth center, Latitude/longitude are ignored and this function returns an ECEF-Oriented transform. 
 	 */
 	UFUNCTION(BlueprintCallable, Category="Celestial Vault|Earth")
-	static FTransform GetEarthCenterTransformUEFrame(double Latitude, double Longitude, double Altitude);
+	static FTransform GetEarthCenterTransformUEFrame(double Latitude, double Longitude, double Altitude, bool Geocentric);
 
 	/** Return the Precession Matrix 
 	 *    The Precession matrix transforms coordinates from one equatorial reference epoch to another by accounting for Earth’s long-term, secular drift of its rotation axis (precession of the equinox).
@@ -340,14 +343,15 @@ public:
 
     /** Compute all Sun Properties for a specific JulianDate */ 
     UFUNCTION(BlueprintCallable, Category="Celestial Vault|Sun")
-    static FStellarBody GetSunInformation(double JulianDate, double ObserverLatitude, double ObserverLongitude);
+    static FStellarBody GetSunInformation(double JulianDate, double ObserverLatitude, double ObserverLongitude, bool bGeoCentric);
 
 	/** Compute all Sun Properties for a specific JulianDate
 	 *
-	 * This override can run faster if we provide it with a cache of the Earth Location and Velocity, because it won't recompute it. 
+	 * This override can run faster if we provide it with a cache of the Earth Location and Velocity, because it won't recompute it.
+	 * if bGeocentric is true, Latitude and Longitude are ignored. 
 	 */ 
 	UFUNCTION(BlueprintCallable, Category="Celestial Vault|Sun")
-	static FStellarBody GetSunInformation_UsingKnownState(double JulianDate, const FPlanetaryBodyKinematicState& EarthKinematicState, double ObserverLatitude, double ObserverLongitude);
+	static FStellarBody GetSunInformation_UsingKnownState(double JulianDate, const FPlanetaryBodyKinematicState& EarthKinematicState, double ObserverLatitude, double ObserverLongitude, bool bGeoCentric);
 
 	
 	/** Returns the Moon Phase for a specific Date
diff --git a/Engine/Plugins/Experimental/CelestialVault/Source/CelestialVault/Public/CelestialVaultDaySequenceActor.h b/Engine/Plugins/Experimental/CelestialVault/Source/CelestialVault/Public/CelestialVaultDaySequenceActor.h
--- a/Engine/Plugins/Experimental/CelestialVault/Source/CelestialVault/Public/CelestialVaultDaySequenceActor.h
+++ b/Engine/Plugins/Experimental/CelestialVault/Source/CelestialVault/Public/CelestialVaultDaySequenceActor.h
@@ -148,13 +148,21 @@ public: // Properties
 	/** Current Time Zone*/ 
 	UPROPERTY(EditAnywhere, BlueprintReadWrite, Category="Location", meta = (ClampMin = "-11", ClampMax = "14"))
 	double GMT_TimeZone = -5.0;
+
+	/**
+	* if true, the UE origin is located at the Planet Center (ECEF) otherwise,
+	* if false, the UE origin is assuming to be defined at one specific point of the planet surface, defined by the Latitude and Longitude properties.
+	**/
+	UPROPERTY(EditAnywhere, BlueprintReadWrite, Category = "Location")
+	bool bGeocentric = false;
+
 	
 	/** Latitude of Level Origin on planet */
-	UPROPERTY(EditAnywhere, BlueprintReadWrite, Category = "Location", meta = (ForceUnits="Degrees", ClampMin = "-90", ClampMax = "90"))
+	UPROPERTY(EditAnywhere, BlueprintReadWrite, Category = "Location", meta = (ForceUnits="Degrees", ClampMin = "-90", ClampMax = "90", EditConditionHides, EditCondition = "bGeocentric==false"))
 	double Latitude = 45.0;
 
 	/** Longitude of Level Origin on planet */
-	UPROPERTY(EditAnywhere, BlueprintReadWrite, Category = "Location", meta = (ForceUnits="Degrees", ClampMin = "-180", ClampMax = "180"))
+	UPROPERTY(EditAnywhere, BlueprintReadWrite, Category = "Location", meta = (ForceUnits="Degrees", ClampMin = "-180", ClampMax = "180", EditConditionHides, EditCondition = "bGeocentric==false"))
 	double Longitude = -73.0;
 
 	/**
@@ -180,7 +188,7 @@ public: // Properties
 #pragma region Celestial Vault related Properties
 
 	/** We generate the sky elements the "Platon" way, using a sphere surrounding the Earth. This is the radius of this sphere. Make sure it's not too small to avoid parallax effects */ 
-	UPROPERTY(EditAnywhere, BlueprintReadWrite, AdvancedDisplay, Category="Celestial Vault", meta=( ForceUnits="Kilometers", ClampMin = "0", ClampMax = "500000"))
+	UPROPERTY(EditAnywhere, BlueprintReadWrite, AdvancedDisplay, Category="Celestial Vault", meta=( ForceUnits="Kilometers", ClampMin = "8000", ClampMax = "500000"))
 	double CelestialVaultDistance = 400000.0;
 
 	/** Percentage of the CelestialVaultDistance at which the Stars are created */ 
@@ -248,11 +256,11 @@ public: // Properties
 	float MoonScale = 2.0f;
 
 	/** If true, the moon Age (Phase) and location can be overriden */  
-	UPROPERTY(EditAnywhere, Category="Moons")
+	UPROPERTY(EditAnywhere, BlueprintReadWrite, Category="Moons")
 	bool bManualControl = false;
 	
 	/** Lunar age. 0 = New Moon, 0.25 = First quarter, 0.5 = Full Moon, 1 = Next New Mo */  
-	UPROPERTY(EditAnywhere, Category="Moons", meta = (EditConditionHides, EditCondition = "bManualControl==true", ClampMin = "0", ClampMax = "1"))
+	UPROPERTY(EditAnywhere, BlueprintReadWrite, Category="Moons", meta = (EditConditionHides, EditCondition = "bManualControl==true", ClampMin = "0", ClampMax = "1"))
 	float MoonAge = 0.2f;
 	
 	/** When faking the moon location, we need to give a location relative to the sun
@@ -263,12 +271,12 @@ public: // Properties
 	 *   (if HoursBehindSun < 12, we are ~1st quarter, so Age should be < 0.5, and if 12< < 24 age should be > 0.5) 
 	 * Any combination is possible, but manual control can lead to inconsistencies. 
 	 */  
-	UPROPERTY(EditAnywhere, Category="Moons", meta = (EditConditionHides, DisplayName="Moon Offset to Sun's Right Ascension", EditCondition = "bManualControl==true", ForceUnits="Hours", ClampMin = "0", ClampMax = "24"))
+	UPROPERTY(EditAnywhere, BlueprintReadWrite, Category="Moons", meta = (EditConditionHides, DisplayName="Moon Offset to Sun's Right Ascension", EditCondition = "bManualControl==true", ForceUnits="Hours", ClampMin = "0", ClampMax = "24"))
 	float MoonOffset_RA = 12.0f;
 
 	/** When faking the moon location, we need to give a location relative to the sun
 	 * This is a way to control this "Vertically" using an offset in Declination. */ 
-	UPROPERTY(EditAnywhere, Category="Moons", meta = (EditConditionHides, DisplayName="Moon Offset to Sun's Declination", EditCondition = "bManualControl==true", ForceUnits="Degrees", ClampMin = "-45", ClampMax = "45"))
+	UPROPERTY(EditAnywhere, BlueprintReadWrite, Category="Moons", meta = (EditConditionHides, DisplayName="Moon Offset to Sun's Declination", EditCondition = "bManualControl==true", ForceUnits="Degrees", ClampMin = "-45", ClampMax = "45"))
 	float MoonOffset_DEC = 15.0f;
 
 	/** Base Sun Intensity
@@ -300,6 +308,10 @@ public:
 	
 public: // Functions
 
+	/** Reset the sequence to its initial state: Initial Day + Initial Time of Day */ 
+	UFUNCTION(BlueprintCallable, Category="Playback")
+	void Stop();
+	
 	/** Returns the current defined day, without any Time, because the Time will be controlled by the DaySequence Time of day - Uses "Now" or the Year/Month/Day properties */ 
 	UFUNCTION(BlueprintCallable, Category = "Celestial Vault|Date")
 	FDateTime GetDate() const;
diff --git a/Engine/Plugins/Experimental/CelestialVault/Source/CelestialVaultEditor/Private/CelestialVaultTests.cpp b/Engine/Plugins/Experimental/CelestialVault/Source/CelestialVaultEditor/Private/CelestialVaultTests.cpp
--- a/Engine/Plugins/Experimental/CelestialVault/Source/CelestialVaultEditor/Private/CelestialVaultTests.cpp
+++ b/Engine/Plugins/Experimental/CelestialVault/Source/CelestialVaultEditor/Private/CelestialVaultTests.cpp
@@ -347,7 +347,7 @@ bool CelestialVault_RADEC::RunTest(const FString& Parameters)
 			double DistanceEarthToSunAU = 0.0;
 
 			// Compute Values
-			UCelestialMaths::GetBodyCelestialCoordinatesAU(Case.JulianDate, BodyCases.VSOP87BodyType, Case.Latitude, Case.Longitude, false,  RAJ2000Hours, DECJ2000Degrees, RAHours, DECDegrees, DistanceBodyToEarthAU, DistanceBodyToSunAU, DistanceEarthToSunAU);
+			UCelestialMaths::GetBodyCelestialCoordinatesAU(Case.JulianDate, BodyCases.VSOP87BodyType, Case.Latitude, Case.Longitude, false,  false, RAJ2000Hours, DECJ2000Degrees, RAHours, DECDegrees, DistanceBodyToEarthAU, DistanceBodyToSunAU, DistanceEarthToSunAU);
 			double RAJ2000Degrees = RAJ2000Hours * 15;
 			double RADegrees = RAHours * 15;
 
