Added Utility Blueprints: PlanetaryBody Tracker and Queries Editor Utility Widget
Moved all the source assets files from my private folder to the CelestialVaultDev project

Stars are now placed at their RA/DEC of Date (was J2000 previously)

Fixed bugs in the planet tracking features

#virtualized
#Synced-CL 50847748

diff --git a/Engine/Plugins/Experimental/CelestialVault/Source/CelestialVault/Private/CelestialVaultDaySequenceActor.cpp b/Engine/Plugins/Experimental/CelestialVault/Source/CelestialVault/Private/CelestialVaultDaySequenceActor.cpp
--- a/Engine/Plugins/Experimental/CelestialVault/Source/CelestialVault/Private/CelestialVaultDaySequenceActor.cpp
+++ b/Engine/Plugins/Experimental/CelestialVault/Source/CelestialVault/Private/CelestialVaultDaySequenceActor.cpp
@@ -673,7 +673,7 @@ bool ACelestialVaultDaySequenceActor::GetClosestPlanetaryBody(FVector ObserverLo
 bool ACelestialVaultDaySequenceActor::GetPlanetaryBodyByVSOP87Type(EVSOP87BodyType VSOP87Type, FPlanetaryBody& FoundPlanetaryBodyInfo, FTransform& BodyTransform)
 {
 	BodyTransform = FTransform::Identity;
-	
+		
 	if (VSOP87Type == EVSOP87BodyType::Moon && MoonDiscComponent) 
 	{
 		BodyTransform = MoonDiscComponent->GetComponentTransform();
@@ -696,6 +696,8 @@ bool ACelestialVaultDaySequenceActor::GetPlanetaryBodyByVSOP87Type(EVSOP87BodyTy
 				if (PlanetaryBodyInfo.VSOP87BodyType == VSOP87Type)
 				{
 					FoundPlanetaryBodyInfo = PlanetaryBodyInfo;
+					FTransform PlanetComponentTransform = PlanetsComponent->GetComponentTransform();
+					BodyTransform = FoundPlanetaryBodyInfo.UELocalTransform * PlanetComponentTransform;
 					return true;
 				}
 			}
@@ -881,15 +883,39 @@ void ACelestialVaultDaySequenceActor::InitStarsComponent()
 			}
 			else
 			{
+				FDateTime LocalTime = GetDateAndTime();
+				FDateTime UTCTime = UCelestialMaths::LocalTimeToUTCTime(LocalTime, GMT_TimeZone, IsDaylightSavingsNow());
+				double JulianDate = UCelestialMaths::UTCDateTimeToJulianDate(UTCTime);
+				FPlanetaryBodyKinematicState EarthKinematicState = UCelestialMaths::GetPlanetaryBodyKinematicState_AU(JulianDate, EVSOP87BodyType::Earth);
+				// Compute the Precession and Nutation Matrices - The formulas expect a Date in Terrestrial Time
+				double TAISeconds = UCelestialMaths::JulianDateToInternationalAtomicTime(JulianDate);
+				double TTSeconds = UCelestialMaths::InternationalAtomicTimeToTerrestrialTime(TAISeconds);
+				double JulianDateTT = UCelestialMaths::SecondsToDay(TTSeconds);
+				FMatrix PrecessionMatrix = UCelestialMaths::GetPrecessionMatrix(JulianDateTT);
+				FMatrix NutationMatrix = UCelestialMaths::GetNutationMatrix(JulianDateTT);
+				
 				// Generate Celestial Stars
 				CelestialStarCatalog->ForeachRow<FCelestialStarInputData>("AEarthDaySequenceActor::InitStarsComponent", [&](const FName& Key, const FCelestialStarInputData& CelestialStarInputData)
 					{
 						if (CelestialStarInputData.Magnitude <= MaxVisibleMagnitude)
 						{
+							// Compute the RA and DEC of Date
+							FVector StarLocation_TopoCentric_FK5J2000_AU = UCelestialMaths::RADECToXYZ_RH(CelestialStarInputData.RA * 15.0, CelestialStarInputData.DEC, CelestialStarInputData.DistanceInPC);
+							FVector StarLocation_TopoCentric_FK5J2000_Aberrated_AU = UCelestialMaths::ComputeAberration(StarLocation_TopoCentric_FK5J2000_AU, EarthKinematicState.Velocity_FK5J2000);
+							// Transform the Body Topographic location with the Precession and Nutation Matrices at Date. 
+							FVector StarLocation_TopoCentric_FK5J2000_Precessed_AU = PrecessionMatrix.InverseTransformVector(StarLocation_TopoCentric_FK5J2000_Aberrated_AU);
+							FVector StarLocation_TopoCentric_FK5J2000_OfDate_AU = NutationMatrix.InverseTransformVector(StarLocation_TopoCentric_FK5J2000_Precessed_AU);
+							double RADegrees = 0.0;
+							double DECDegrees = 0.0;
+							double DistanceStarToEarthAU = 0.0;
+							UCelestialMaths::XYZToRADEC_RH(StarLocation_TopoCentric_FK5J2000_OfDate_AU, RADegrees, DECDegrees, DistanceStarToEarthAU);
+							double RAHours = RADegrees / 15.0;
+
+							
 							// Location - Convert to UE Left-handed Frame (Invert the Y coordinate)
 							// Special case for the Stars. To avoid recomputing the ISM bounds, we use the parent bounds, and therefore inherit from the transformation of the Celestial Vault.
 							// The mesh being 100 units radius, we have a local transform corresponding to a StarsVaultPercentage distance. 
-							FVector StarLocation = UCelestialMaths::RADECToXYZ_RH(CelestialStarInputData.RA * 15.0, CelestialStarInputData.DEC, StarsVaultPercentage) * FVector(1.0, -1.0, 1.0);
+							FVector StarLocation = UCelestialMaths::RADECToXYZ_RH(RADegrees, DECDegrees, StarsVaultPercentage) * FVector(1.0, -1.0, 1.0);
 
 							// Color
 							FLinearColor StarColor = UCelestialMaths::BVtoLinearColor(CelestialStarInputData.ColorIndex);
@@ -907,8 +933,10 @@ void ACelestialVaultDaySequenceActor::InitStarsComponent()
 							if (bKeepStarsInfo)
 							{
 								FStellarBody StarInfo;
-								StarInfo.RA = CelestialStarInputData.RA;
-								StarInfo.DEC = CelestialStarInputData.DEC;
+								StarInfo.RAJ2000 = CelestialStarInputData.RA;
+								StarInfo.DECJ2000 = CelestialStarInputData.DEC;
+								StarInfo.RA = RAHours;
+								StarInfo.DEC = DECDegrees;
 								StarInfo.DistanceInAU = UCelestialMaths::ParsecsToAstronomicalUnits(CelestialStarInputData.DistanceInPC);
 								StarInfo.Name = CelestialStarInputData.Name;
 								StarInfo.Magnitude = CelestialStarInputData.Magnitude;
@@ -1050,6 +1078,7 @@ void ACelestialVaultDaySequenceActor::InitPlanetsComponent()
 FPlanetaryBody ACelestialVaultDaySequenceActor::GetPlanetaryBodyInfo_UsingKnownState(double JulianDay, const FPlanetaryBodyInputData& InputPlanetaryBody, const FPlanetaryBodyKinematicState& EarthKinematicState, bool bIgnoreRelativisticEffect) const
 {
 	FPlanetaryBody BodyInfo;
+	BodyInfo.VSOP87BodyType = InputPlanetaryBody.VSOP87BodyType;
 	BodyInfo.OrbitType = InputPlanetaryBody.OrbitType;
 	BodyInfo.Name = InputPlanetaryBody.Name;
 	BodyInfo.Radius = InputPlanetaryBody.Radius;
