[Celestial Vault]
Added an optional system to adapt the SkyAtmosphere radius depending on Camera location, for globe-scale applications.
Added the isSky option to the CelestialVault Material to have it captured by the Skylight. Added a proxy mesh to write Velocities and depth to work around their absence when isSky is on
Added a custom Cloud Material using an MPC to read the actual time of day. It allows for time-lapses.
Added a new Star Shader with Conservative Energy for Stars. It produces better-looking magnitudes, while still allowing artistic control over size and intensity.
Added an editor utility widget to compute the Star Mask texture mean, for proper energy conservation.

Moved some hardcoded constants to a WGS84 class
Minor fixes to some internal blueprints
Exposed the Moon Normal texture to the Material Instance

#virtualized
#Synced-CL 51124650

diff --git a/Engine/Plugins/Experimental/CelestialVault/Source/CelestialVault/Private/CelestialMaths.cpp b/Engine/Plugins/Experimental/CelestialVault/Source/CelestialVault/Private/CelestialMaths.cpp
--- a/Engine/Plugins/Experimental/CelestialVault/Source/CelestialVault/Private/CelestialMaths.cpp
+++ b/Engine/Plugins/Experimental/CelestialVault/Source/CelestialVault/Private/CelestialMaths.cpp
@@ -36,6 +36,11 @@ double const UCelestialMaths::ParsecAstronomicalUnits = 206264.80624709636;
 double const UCelestialMaths::NewMoonReferenceJulianDate = 2460705.025196759; // Known new Moon was January 29th, 2025, at 12:36:17 UTC
 double const UCelestialMaths::SynodicMonthAverage = 29.530588853;
 
+
+double UCelestialMaths::WGS84::SemiMajorAxis = 6378137.0;
+double UCelestialMaths::WGS84::SemiMinorAxis = 6356752.314245;
+double UCelestialMaths::WGS84::Flattening = 1.0 / 298.257223563;
+
 #pragma endregion
 
 #pragma region Colors
@@ -713,12 +718,10 @@ FVector UCelestialMaths::GeodeticLatLonToECEFXYZAU(double Latitude, double Longi
 	double LatitudeRadians = FMath::DegreesToRadians(Latitude);
 	double LongitudeRadians = FMath::DegreesToRadians(Longitude);
 	
-	const double a = MetersToAstronomicalUnits(6378136.6);
-	const double f = 1.0 / 298.25642;
-
-	const double C = 1.0 / FMath::Sqrt(FMath::Cos(LatitudeRadians) * FMath::Cos(LatitudeRadians) + (1.0 - f) * (1.0 - f) * (FMath::Sin(LatitudeRadians) * FMath::Sin(LatitudeRadians)));
+	const double a = MetersToAstronomicalUnits(WGS84::SemiMajorAxis);
+	const double C = 1.0 / FMath::Sqrt(FMath::Cos(LatitudeRadians) * FMath::Cos(LatitudeRadians) + (1.0 - WGS84::Flattening) * (1.0 - WGS84::Flattening) * (FMath::Sin(LatitudeRadians) * FMath::Sin(LatitudeRadians)));
 
-	const double S = (1.0 - f) * (1.0 - f) * C;
+	const double S = (1.0 - WGS84::Flattening) * (1.0 - WGS84::Flattening) * C;
 	const double h = MetersToAstronomicalUnits(Altitude);
 
 	return FVector(
@@ -727,6 +730,55 @@ FVector UCelestialMaths::GeodeticLatLonToECEFXYZAU(double Latitude, double Longi
 		(a * S + h) * FMath::Sin(LatitudeRadians)); 
 }
 
+void UCelestialMaths::ECEFXYZAUToGeodeticLatLon(FVector ECEFLocationAU, double &Latitude, double &Longitude, double &AltitudeMeters)
+{
+	// Convert coordinates from AU to meters
+	const double Xm = AstronomicalUnitsToMeters(ECEFLocationAU.X);
+	const double Ym = AstronomicalUnitsToMeters(ECEFLocationAU.Y);
+	const double Zm = AstronomicalUnitsToMeters(ECEFLocationAU.Z);
+	const double ProjectedRadius = FMath::Sqrt(Xm * Xm + Ym * Ym);
+	const double e2 = 1.0 - (WGS84::SemiMinorAxis * WGS84::SemiMinorAxis) / (WGS84::SemiMajorAxis * WGS84::SemiMajorAxis);  // first eccentricity squared
+	const double ep2 = (WGS84::SemiMajorAxis * WGS84::SemiMajorAxis) / (WGS84::SemiMinorAxis * WGS84::SemiMinorAxis) - 1.0; // second eccentricity squared
+
+	// Longitude
+	const double LongitudeRadians = FMath::Atan2(Ym, Xm);
+	Longitude = FMath::RadiansToDegrees(LongitudeRadians);
+	
+	// Bowring formula
+	const double Theta = FMath::Atan2(Zm * WGS84::SemiMajorAxis, ProjectedRadius * WGS84::SemiMinorAxis);
+	const double SinTheta = FMath::Sin(Theta);
+	const double CosTheta = FMath::Cos(Theta);
+
+	// Latitude 
+	const double LatitudeRadians = FMath::Atan2(Zm + ep2 * WGS84::SemiMinorAxis * SinTheta * SinTheta * SinTheta,  ProjectedRadius - e2 * WGS84::SemiMajorAxis * CosTheta * CosTheta * CosTheta);
+	Latitude = FMath::RadiansToDegrees(LatitudeRadians);
+
+	// Altitude
+	if (ProjectedRadius > UE_SMALL_NUMBER)
+	{
+		const double SinLatitude = FMath::Sin(LatitudeRadians);
+		const double GeocentricNormal = WGS84::SemiMajorAxis / FMath::Sqrt(1.0 - e2 * SinLatitude * SinLatitude);
+		AltitudeMeters = ProjectedRadius / FMath::Cos(LatitudeRadians) - GeocentricNormal;	
+	}
+	else
+	{
+		AltitudeMeters = FMath::Abs(Zm) - WGS84::SemiMinorAxis; 
+	}
+	
+}
+
+double UCelestialMaths::WGS84GeocentricRadius(double Latitude)
+{
+	const double Phi = FMath::DegreesToRadians(Latitude);
+	const double CosPhi = FMath::Cos(Phi);
+	const double SinPhi = FMath::Sin(Phi);
+
+	const double Numerator = FMath::Square(WGS84::SemiMajorAxis * WGS84::SemiMajorAxis * CosPhi) + FMath::Square(WGS84::SemiMinorAxis * WGS84::SemiMinorAxis * SinPhi);
+	const double Denominator = FMath::Square(WGS84::SemiMajorAxis * CosPhi) + FMath::Square(WGS84::SemiMinorAxis * SinPhi);
+
+	return FMath::Sqrt(Numerator / Denominator);
+}
+
 FVector UCelestialMaths::GetObserverGeocentricLocationAU(double Latitude,double Longitude, double Altitude, double JulianDate )
 {
 	const FVector ObserverECEF = GeodeticLatLonToECEFXYZAU(Latitude, Longitude, Altitude);
@@ -857,10 +909,7 @@ FTransform UCelestialMaths::GetEarthCenterTransformECEF(double Latitude, double
 		double Tolerance = 1.E-50; // Normalize with a very low threshold, because default is 10-8, too high for double computations
 
 		// Compute the ellipsoid normal (Earth...)
-		double SemiMajorMeter = 6378137.0;
-		double SemiMinorMetre = 6356752.3142451793;
-		
-		FVector OneOverRadiiSquared = FVector(1.0 / (SemiMajorMeter * SemiMajorMeter), 1.0 / (SemiMajorMeter * SemiMajorMeter), 1.0 / (SemiMinorMetre * SemiMinorMetre)); 
+		FVector OneOverRadiiSquared = FVector(1.0 / (WGS84::SemiMajorAxis * WGS84::SemiMajorAxis), 1.0 / (WGS84::SemiMajorAxis * WGS84::SemiMajorAxis), 1.0 / (WGS84::SemiMinorAxis * WGS84::SemiMinorAxis)); 
 		FVector GeodeticSurfaceNormal( ECEFLocation.X * OneOverRadiiSquared.X, ECEFLocation.Y * OneOverRadiiSquared.Y, ECEFLocation.Z * OneOverRadiiSquared.Z);
 		GeodeticSurfaceNormal.Normalize(Tolerance);
 
diff --git a/Engine/Plugins/Experimental/CelestialVault/Source/CelestialVault/Private/CelestialVaultDaySequenceActor.cpp b/Engine/Plugins/Experimental/CelestialVault/Source/CelestialVault/Private/CelestialVaultDaySequenceActor.cpp
--- a/Engine/Plugins/Experimental/CelestialVault/Source/CelestialVault/Private/CelestialVaultDaySequenceActor.cpp
+++ b/Engine/Plugins/Experimental/CelestialVault/Source/CelestialVault/Private/CelestialVaultDaySequenceActor.cpp
@@ -18,6 +18,8 @@
 #include "UObject/ConstructorHelpers.h"
 #include "DaySequenceSubsystem.h"
 #include "DaySequenceCollectionAsset.h"
+#include "Materials/MaterialParameterCollectionInstance.h"
+#include "Materials/MaterialParameterCollection.h"
 #include "DaySequence/Private/DaySequencePlayer.h"
 #include "Engine/World.h"
 #include "Engine/StaticMesh.h"
@@ -27,6 +29,8 @@
 #include "CelestialMaths.h"
 #include "CelestialDataTypes.h"
 #include "CelestialVault.h"
+#include "ViewportCameraProvider.h"
+
 
 #if CELESTIAL_VAULT_ENABLE_DRAW_DEBUG
 
@@ -64,13 +68,14 @@ ACelestialVaultDaySequenceActor::ACelestialVaultDaySequenceActor(const FObjectIn
 	{
 		GlobalPostProcessVolume->SetupAttachment(RootComponent);
 		GlobalPostProcessVolume->Settings.bOverride_AutoExposureMinBrightness = true;
-		GlobalPostProcessVolume->Settings.AutoExposureMinBrightness = -0.5f;
+		GlobalPostProcessVolume->Settings.AutoExposureMinBrightness = 0.25f;
 		GlobalPostProcessVolume->Settings.bOverride_LocalExposureHighlightContrastCurve = true;
 		GlobalPostProcessVolume->Settings.LocalExposureHighlightContrastCurve = HighlightContrastCurve.Object;
-		GlobalPostProcessVolume->Settings.bOverride_LocalExposureDetailStrength = true;
-		GlobalPostProcessVolume->Settings.LocalExposureDetailStrength = 1.2f;
-		GlobalPostProcessVolume->Settings.bOverride_LocalExposureBlurredLuminanceBlend = true;
-		GlobalPostProcessVolume->Settings.LocalExposureBlurredLuminanceBlend = 0.4f;
+
+		GlobalPostProcessVolume->Settings.bOverride_AutoExposureSpeedUp = true;
+		GlobalPostProcessVolume->Settings.AutoExposureSpeedUp = 100.0f;
+		GlobalPostProcessVolume->Settings.bOverride_AutoExposureSpeedDown = true;
+		GlobalPostProcessVolume->Settings.AutoExposureSpeedDown = 100.0f;
 	}
 	
 	// Components Attached to Root
@@ -90,7 +95,7 @@ ACelestialVaultDaySequenceActor::ACelestialVaultDaySequenceActor(const FObjectIn
 		if (!IsTemplate())
 		{
 			// We don't want to load this material for the CDO as it will hold on to it forever and it is quite a large asset.
-			static ConstructorHelpers::FObjectFinder<UMaterialInterface> VolumetricCloudDefaultMaterialRef(TEXT("/Engine/EngineSky/VolumetricClouds/m_SimpleVolumetricCloud_Inst.m_SimpleVolumetricCloud_Inst"));
+			static ConstructorHelpers::FObjectFinder<UMaterialInterface> VolumetricCloudDefaultMaterialRef(TEXT("/CelestialVault/Materials/m_SimpleVolumetricCloud_TOD_Inst.m_SimpleVolumetricCloud_TOD_Inst"));
 			VolumetricCloudComponent->SetMaterial(VolumetricCloudDefaultMaterialRef.Object);
 		}	
 	}
@@ -166,11 +171,28 @@ ACelestialVaultDaySequenceActor::ACelestialVaultDaySequenceActor(const FObjectIn
 		DeepSkyComponent->SetMaterial(0, SkySphereDefaultMaterial.Object.Get());
 		DeepSkyComponent->SetRelativeScale3D(FVector(CelestialVaultDistance*1000.0));	
 	}
+
+
+	// Deep Sky background
+	static ConstructorHelpers::FObjectFinder<UMaterialInterface> VelocityProxyDefaultMaterial(TEXT("/CelestialVault/Materials/M_VelocityProxy.M_VelocityProxy"));
+	VelocityVectorsProxyComponent = CreateOptionalDefaultSubobject<UStaticMeshComponent>(TEXT("Velocity Vectors Proxy"));
+	if (VelocityVectorsProxyComponent)
+	{
+		VelocityVectorsProxyComponent->SetupAttachment(DeepSkyComponent);
+		VelocityVectorsProxyComponent->SetStaticMesh(SkySphereDefaultMesh.Object);
+		VelocityVectorsProxyComponent->SetCollisionEnabled(ECollisionEnabled::NoCollision);
+		VelocityVectorsProxyComponent->SetGenerateOverlapEvents(false);
+		VelocityVectorsProxyComponent->SetCastShadow(false);
+		VelocityVectorsProxyComponent->SetAffectDynamicIndirectLighting(false);
+		VelocityVectorsProxyComponent->SetCanEverAffectNavigation(false);
+		VelocityVectorsProxyComponent->SetMaterial(0, VelocityProxyDefaultMaterial.Object.Get());
+		VelocityVectorsProxyComponent->SetRelativeScale3D(FVector(1.001));	
+	}
 	
 
 	// Stars ISM
 	static ConstructorHelpers::FObjectFinder<UStaticMesh> PlaneXMesh(TEXT("/CelestialVault/Meshes/SM_Plane_FacingX.SM_Plane_FacingX"));
-	static ConstructorHelpers::FObjectFinder<UMaterialInterface> StarsDefaultMaterial(TEXT("/CelestialVault/Materials/MI_Stars.MI_Stars"));
+	static ConstructorHelpers::FObjectFinder<UMaterialInterface> StarsDefaultMaterial(TEXT("/CelestialVault/Materials/MI_StarsStar_EnergyConservative.MI_StarsStar_EnergyConservative"));
 	StarsComponent = CreateOptionalDefaultSubobject<UInstancedStaticMeshComponent>(TEXT("Stars"));
 	if (StarsComponent)
 	{
@@ -257,6 +279,13 @@ ACelestialVaultDaySequenceActor::ACelestialVaultDaySequenceActor(const FObjectIn
 		static ConstructorHelpers::FObjectFinder<UDataTable> DefaultPlanetaryBodiesCatalog(TEXT("/CelestialVault/Data/DT_SolarSystemPlanets.DT_SolarSystemPlanets"));
 		PlanetsCatalog = DefaultPlanetaryBodiesCatalog.Object.Get();
 	}
+
+	// MPC
+	static ConstructorHelpers::FObjectFinder<UMaterialParameterCollection> MPC(TEXT("/CelestialVault/Materials/MPC_CelestialVault.MPC_CelestialVault"));
+	CelestialVaultMPC = MPC.Object.Get();
+
+	SetTimePerCycle(24);
+	SetInitialTimeOfDay(12);
 }
 
 bool ACelestialVaultDaySequenceActor::IsDaylightSavingsNow() const
@@ -334,6 +363,50 @@ void ACelestialVaultDaySequenceActor::Tick(float DeltaTime)
 		}
 	}
 #endif
+
+	if (CelestialVaultMPC)
+	{
+		if (UWorld* World = GetWorld())
+		{
+			if (UMaterialParameterCollectionInstance* CelestialVaultMPCInstance = World->GetParameterCollectionInstance(CelestialVaultMPC))
+			{
+				CelestialVaultMPCInstance->SetScalarParameterValue(TEXT("TimeOfDay"), (SequenceLoopCount * 24 + GetTimeOfDay())*3600);
+			}	
+		}
+	}
+	
+	if (bAdjustSkyAtmosphereToLocalRadius)
+	{
+		if (IViewportCameraProvider* Provider = GetCameraProvider())
+		{
+			FVector CameraLocation;
+			FRotator CameraRotation;
+
+			if (Provider->GetCamera(CameraLocation, CameraRotation))
+			{
+				// Check if we moved by more than AdjustSkyAtmosphereMovementThreshold
+				if ( FVector::DistSquared(LastViewportCameraLocation, CameraLocation) >  (AdjustSkyAtmosphereMovementThreshold * AdjustSkyAtmosphereMovementThreshold * UE_ONE_KILOMETER * UE_ONE_KILOMETER))
+				{
+					LastViewportCameraLocation = CameraLocation;
+
+					if (PlanetCenterComponent)
+					{
+						FVector CameraECEFLocationMeters = PlanetCenterComponent->GetComponentTransform().InverseTransformPosition(CameraLocation) * FVector(0.01, -0.01, 0.01);
+						const FVector CameraECEFLocationAU =  CameraECEFLocationMeters * UCelestialMaths::MetersToAstronomicalUnits(1);
+
+						double CameraLatitude, CameraLongitude, CameraAltitudeMeters;
+						UCelestialMaths::ECEFXYZAUToGeodeticLatLon(CameraECEFLocationAU, CameraLatitude, CameraLongitude, CameraAltitudeMeters);
+						//UE_LOG(LogCelestialVault, Verbose, TEXT("Current Camera Location: Latitude=%f, Longitude=%f, Altitude=%f m"), CameraLatitude, CameraLongitude, CameraAltitudeMeters);
+
+						double WG84SRadiusMeters = UCelestialMaths::WGS84GeocentricRadius(CameraLatitude);
+						double AdjustedRadiusKm = FMath::GetMappedRangeValueClamped(FVector2D(AdjustSkyAtmosphereFadeMinAltitude * 1000.0, AdjustSkyAtmosphereFadeMaxAltitude * 1000.0), FVector2D(WG84SRadiusMeters, UCelestialMaths::WGS84::SemiMajorAxis), CameraAltitudeMeters) / 1000.0;
+						//UE_LOG(LogCelestialVault, Verbose, TEXT("WG84 Radius: %f km - Retained = %f km"), WG84SRadiusMeters / 1000.0, AdjustedRadiusKm);
+						SkyAtmosphereComponent->SetBottomRadius(AdjustedRadiusKm);
+					}
+				}
+			}
+		}
+	}
 }
 
 void ACelestialVaultDaySequenceActor::SequencePlayerUpdated(float CurrentTime, float PreviousTime)
@@ -764,7 +837,12 @@ void ACelestialVaultDaySequenceActor::PostEditChangeProperty(FPropertyChangedEve
 		PropertyName == GET_MEMBER_NAME_CHECKED(ACelestialVaultDaySequenceActor, Day) ||
 		PropertyName == GET_MEMBER_NAME_CHECKED(ACelestialVaultDaySequenceActor, GMT_TimeZone) ||
 		PropertyName == GET_MEMBER_NAME_CHECKED(ACelestialVaultDaySequenceActor, Latitude) ||
-		PropertyName == GET_MEMBER_NAME_CHECKED(ACelestialVaultDaySequenceActor, Longitude))
+		PropertyName == GET_MEMBER_NAME_CHECKED(ACelestialVaultDaySequenceActor, Longitude ||
+		PropertyName == GET_MEMBER_NAME_CHECKED(ACelestialVaultDaySequenceActor, DaylightSavingsMode) ||
+		PropertyName == GET_MEMBER_NAME_CHECKED(ACelestialVaultDaySequenceActor, DaylightSavingsStart) ||
+		PropertyName == GET_MEMBER_NAME_CHECKED(ACelestialVaultDaySequenceActor, DaylightSavingsEnd) ||
+		PropertyName == GET_MEMBER_NAME_CHECKED(ACelestialVaultDaySequenceActor, DaylightSavingsSwitchHour) ||
+		PropertyName == GET_MEMBER_NAME_CHECKED(ACelestialVaultDaySequenceActor, bGeocentric)))
 	{
 		// DateTime has changed, we need to update the Reference CelestialVault Angle.
 		FDateTime LocalTime0 = GetDate();
@@ -1173,7 +1251,7 @@ void ACelestialVaultDaySequenceActor::UpdateBodiesMotion()
 		SunLightComponent->SetIntensity(SunLightIntensity);
 	}
 
-	// Don't use a curve cache for the Celestial Vault rotation. It's a basic angular rotation. 
+	// Don't use a curve cache for the Celestial Vault rotation. It's a basic angular rotation.
 	double VaultYaw = GetCelestialVaultAngle(TimeOfDay );
 	CelestialVaultComponent->SetRelativeRotation(FRotator(0.0, VaultYaw, 0.0));
 }
diff --git a/Engine/Plugins/Experimental/CelestialVault/Source/CelestialVault/Private/ViewportCameraProvider.cpp b/Engine/Plugins/Experimental/CelestialVault/Source/CelestialVault/Private/ViewportCameraProvider.cpp
new file mode 100644
--- /dev/null
+++ b/Engine/Plugins/Experimental/CelestialVault/Source/CelestialVault/Private/ViewportCameraProvider.cpp
@@ -0,0 +1,42 @@
+﻿// Copyright Epic Games, Inc. All Rights Reserved.
+
+
+#include "ViewportCameraProvider.h"
+#include "Engine/Engine.h"
+#include "GameFramework/PlayerController.h"
+
+static IViewportCameraProvider* GCameraProvider = nullptr;
+
+
+
+//////////////////////////////////////////////////////////////////////////
+
+bool FRuntimeCameraProvider::GetCamera(FVector& OutLocation, FRotator& OutRotation) const
+{
+	if (!GEngine) return false;
+
+	UWorld* World = GEngine->GetCurrentPlayWorld();
+	if (!World) return false;
+
+	APlayerController* PC = World->GetFirstPlayerController();
+	if (!PC) return false;
+
+	PC->GetPlayerViewPoint(OutLocation, OutRotation);
+	return true;
+}
+
+IViewportCameraProvider* GetCameraProvider()
+{
+	if (!GCameraProvider)
+	{
+		static FRuntimeCameraProvider DefaultProvider;
+		GCameraProvider = &DefaultProvider;
+	}
+
+	return GCameraProvider;
+}
+
+void SetCameraProvider(IViewportCameraProvider* InProvider)
+{
+	GCameraProvider = InProvider;
+}
\ No newline at end of file
diff --git a/Engine/Plugins/Experimental/CelestialVault/Source/CelestialVault/Public/CelestialMaths.h b/Engine/Plugins/Experimental/CelestialVault/Source/CelestialVault/Public/CelestialMaths.h
--- a/Engine/Plugins/Experimental/CelestialVault/Source/CelestialVault/Public/CelestialMaths.h
+++ b/Engine/Plugins/Experimental/CelestialVault/Source/CelestialVault/Public/CelestialMaths.h
@@ -264,6 +264,18 @@ public:
 
 #pragma region Earth
 
+	class WGS84
+	{
+	public:
+		/**  Semi Major Axis in Meters **/ 
+		static double SemiMajorAxis;
+
+		/**  Semi Minor Axis in Meters **/
+		static double SemiMinorAxis;
+		
+		static double Flattening;
+	};
+	
     /** Retrurn the Earth Rotation Angle (In Degrees) as measured by GMST (Greenwich Mean Sidereal Time)
      *
      * It refers to the angle of Earth's rotation relative to the fixed stars, specifically the hour angle of the mean vernal equinox as observed from the Greenwich meridian.
@@ -282,6 +294,26 @@ public:
     UFUNCTION(BlueprintCallable, Category="Celestial Vault|Earth")
     static FVector GeodeticLatLonToECEFXYZAU(double Latitude, double Longitude, double Altitude);
 
+
+    /** Converts a ECEF Location for the WGS84 Ellipsoid into Geodetic Latitude and Longitude.  
+     *
+	 *   Be careful, XYZ coordinates are
+	 *    - ECEF Coordinates in the ECEF Right-Handed Frame (not the Left-handed UE ones in UE Units)
+	 *    - Expressed in Astronomical Units (AU)
+	 *
+	 *   Altitude is returned as Meters
+     */
+    UFUNCTION(BlueprintCallable, Category="Celestial Vault|Earth")
+	static void ECEFXYZAUToGeodeticLatLon(FVector ECEFLocationAU, double &Latitude, double &Longitude, double &AltitudeMeters);
+
+    /**
+     *  Return the WGS84 Radius (in Meters) for a specific Latitude 
+     */
+    UFUNCTION(BlueprintCallable, Category="Celestial Vault|Earth")
+	static double WGS84GeocentricRadius(double Latitude);
+
+    
+
     /** Return the Geocentric position of an observer located at the Earth surface, considering the rotation at this specific JulianDate, using the Greenwich Apparent Didereal Time.
      *
      *  The position is expressed relatively to the earth center, but on the solar system reference frame.
diff --git a/Engine/Plugins/Experimental/CelestialVault/Source/CelestialVault/Public/CelestialVaultDaySequenceActor.h b/Engine/Plugins/Experimental/CelestialVault/Source/CelestialVault/Public/CelestialVaultDaySequenceActor.h
--- a/Engine/Plugins/Experimental/CelestialVault/Source/CelestialVault/Public/CelestialVaultDaySequenceActor.h
+++ b/Engine/Plugins/Experimental/CelestialVault/Source/CelestialVault/Public/CelestialVaultDaySequenceActor.h
@@ -29,6 +29,7 @@ class UExponentialHeightFogComponent;
 class UPostProcessComponent;
 class UStaticMeshComponent;
 class UInstancedStaticMeshComponent;
+class UMaterialParameterCollection;
 
 
 UCLASS(Blueprintable, HideCategories=(Tags, Networking, LevelInstance))
@@ -80,6 +81,11 @@ public: // Properties
 	UPROPERTY(VisibleDefaultsOnly, BlueprintReadOnly, Category = "Celestial Vault", meta = (AllowPrivateAccess = "true"))
 	TObjectPtr<UStaticMeshComponent> DeepSkyComponent;
 
+	/* Because the Deep Sky material is set to IsSky to be included in the SkyLight capture, it doesn't write any velocity vector, creating visual artifacts (streaks) when the vault rotates fast.
+	 * This component use a simple opaque material to write the velocities and have the background behaving correctly when the vault rotates quickly. */ 
+	UPROPERTY(VisibleDefaultsOnly, BlueprintReadOnly, Category = "Celestial Vault", meta = (AllowPrivateAccess = "true"))
+	TObjectPtr<UStaticMeshComponent> VelocityVectorsProxyComponent;
+
 	UPROPERTY(VisibleDefaultsOnly, BlueprintReadOnly, Category = "Celestial Vault", meta = (AllowPrivateAccess = "true"))
 	TObjectPtr<UStaticMeshComponent> MoonDiscComponent;
 
@@ -99,7 +105,7 @@ public: // Properties
 
 	/** Current Year*/ 
 	UPROPERTY(EditAnywhere, BlueprintReadWrite, Category="Date", meta = (EditConditionHides, EditCondition = "!bUseCurrentDate", ClampMin = "1", ClampMax = "9999"))
-	int Year = 2025;
+	int Year = 2026;
 
 	/** Current Month*/ 
 	UPROPERTY(EditAnywhere, BlueprintReadWrite, Category="Date", meta = (EditConditionHides, EditCondition = "!bUseCurrentDate", ClampMin = "1", ClampMax = "12"))
@@ -107,7 +113,7 @@ public: // Properties
 
 	/** Current Day*/ 
 	UPROPERTY(EditAnywhere, BlueprintReadWrite, Category="Date", meta = (EditConditionHides, EditCondition = "!bUseCurrentDate", ClampMin = "1", ClampMax = "31"))
-	int Day = 7;
+	int Day = 1;
 
 	/** At runtime, if we cycle over more than one day, we increase this one */
 	UPROPERTY(Transient, VisibleAnywhere, BlueprintReadOnly, Category="Date")
@@ -155,7 +161,6 @@ public: // Properties
 	**/
 	UPROPERTY(EditAnywhere, BlueprintReadWrite, Category = "Location")
 	bool bGeocentric = false;
-
 	
 	/** Latitude of Level Origin on planet */
 	UPROPERTY(EditAnywhere, BlueprintReadWrite, Category = "Location", meta = (ForceUnits="Degrees", ClampMin = "-90", ClampMax = "90", EditConditionHides, EditCondition = "bGeocentric==false"))
@@ -176,17 +181,20 @@ public: // Properties
 	double NorthOffset = 0.0;
 	
 	// Greenwich Mean Sidereal Time at corresponding to a 0 Time of Day (midnight in the morning) for the selected Date. 
-	UPROPERTY(BlueprintReadOnly, Category="Location")
+	UPROPERTY(Transient, BlueprintReadOnly, Category="Location")
 	double GMST0_Unwrapped = 0.0;
 
 	// Transform to apply to the planet to have it located tangent to the Origin
-	UPROPERTY(BlueprintReadOnly, VisibleAnywhere, Category="Location")
+	UPROPERTY(Transient, BlueprintReadOnly, VisibleAnywhere, Category="Location")
 	FTransform PlanetCenterTransform;
 	
 #pragma endregion
 
 #pragma region Celestial Vault related Properties
 
+	UPROPERTY(EditAnywhere, BlueprintReadOnly, AdvancedDisplay, Category="Celestial Vault")
+	TObjectPtr<UMaterialParameterCollection> CelestialVaultMPC;
+	
 	/** We generate the sky elements the "Platon" way, using a sphere surrounding the Earth. This is the radius of this sphere. Make sure it's not too small to avoid parallax effects */ 
 	UPROPERTY(EditAnywhere, BlueprintReadWrite, AdvancedDisplay, Category="Celestial Vault", meta=( ForceUnits="Kilometers", ClampMin = "8000", ClampMax = "500000"))
 	double CelestialVaultDistance = 400000.0;
@@ -202,6 +210,32 @@ public: // Properties
 	/** Percentage of the CelestialVaultDistance at which the Moons are created */
 	UPROPERTY(EditAnywhere, BlueprintReadWrite, AdvancedDisplay, Category="Celestial Vault", meta=( ForceUnits="Percent", ClampMin = "0.1", ClampMax = "100"))
 	double MoonVaultPercentage = 95.0;
+
+	/** Make sure the SkyAtmosphere is dynamically adjusted to the local ellipsoid radius. (Not needed for applications staying close to the Origin).  
+	 *
+	 * If true, the system will track the current viewport camera location, extract its Latitude, and adjust the radius of the SkyAtmosphere to the expected ellipsoidal radius
+	 * This is a trick to have the round SkyAtmosphere fitting with the Earth surface wherever you're located on Earth
+	 */
+	UPROPERTY(EditAnywhere, BlueprintReadWrite, AdvancedDisplay, Category="Celestial Vault")
+	bool bAdjustSkyAtmosphereToLocalRadius = false;
+
+	/** We don't need to adust the radius everyframe. On earth, a rought order of magnitude says that the ellipsoid changes by 2.2 meters each km. 
+	 * It's fine to adjust the radius only if we have moved by some dozens of km... 
+	 */
+	UPROPERTY(EditAnywhere, BlueprintReadWrite, AdvancedDisplay, Category="Celestial Vault", meta=( ForceUnits="Km", ClampMin = "1", ClampMax = "2000", EditConditionHides, EditCondition = "bAdjustSkyAtmosphereToLocalRadius==true"))
+	float AdjustSkyAtmosphereMovementThreshold = 10;
+
+	/** For visualization at high altitudes, using the theoretical ellipsoidal radius below the camera is not sufficient, as we can see other areas of the earth.
+	 * Between the AdjustSkyAtmosphereFadeMinAltitude and AdjustSkyAtmosphereFadeMaxAltitude values, we'll progressly scale the SkyAtmosphere radius toward the Ellipsoid major radius to avoid artifacts
+	 */
+	UPROPERTY(EditAnywhere, BlueprintReadWrite, AdvancedDisplay, Category="Celestial Vault", meta=( ForceUnits="Km", ClampMin = "10", ClampMax = "2000", EditConditionHides, EditCondition = "bAdjustSkyAtmosphereToLocalRadius==true"))
+	float AdjustSkyAtmosphereFadeMinAltitude = 100;
+
+	/** For visualization at high altitudes, using the theoretical ellipsoidal radius below the camera is not sufficient, as we can see other areas of the earth.
+	 * Between the AdjustSkyAtmosphereFadeMinAltitude and AdjustSkyAtmosphereFadeMaxAltitude values, we'll progressly scale the SkyAtmosphere radius toward the Ellipsoid major radius to avoid artifacts
+	 */
+	UPROPERTY(EditAnywhere, BlueprintReadWrite, AdvancedDisplay, Category="Celestial Vault", meta=( ForceUnits="Km", ClampMin = "10", ClampMax = "10000", EditConditionHides, EditCondition = "bAdjustSkyAtmosphereToLocalRadius==true"))
+	float AdjustSkyAtmosphereFadeMaxAltitude = 300;
 	
 #pragma endregion
 
@@ -224,7 +258,7 @@ public: // Properties
 	bool bKeepStarsInfo = false;
 	
 	/** Array of the created Stars information - Only populated if KeepStarsInfo is true */ 
-	UPROPERTY(BlueprintReadOnly, Category="Stars")
+	UPROPERTY(Transient, BlueprintReadOnly, Category="Stars")
 	TArray<FStellarBody> Stars;
 	
 #pragma endregion
@@ -244,7 +278,7 @@ public: // Properties
 	bool bKeepPlanetsInfos = false;
 	
 	/** Array of the created planetary bodies, with all their computed information - Only populated if KeepPlanetsInfos is true */ 
-	UPROPERTY(BlueprintReadOnly, Category="Celestial Vault|Planets")
+	UPROPERTY(Transient, BlueprintReadOnly, Category="Celestial Vault|Planets")
 	TArray<FPlanetaryBody> Planets;
 
 #pragma endregion
@@ -415,6 +449,7 @@ private:
 	FInterpCurve<FQuat>  MoonLightRotationCurve;
 	FInterpCurve<float>  MoonAgeCurve;
 	FInterpCurve<FQuat>  SunLightRotationCurve;
+	FVector LastViewportCameraLocation = FVector::ZeroVector;
 	
 #if CELESTIAL_VAULT_ENABLE_DRAW_DEBUG
 
diff --git a/Engine/Plugins/Experimental/CelestialVault/Source/CelestialVault/Public/ViewportCameraProvider.h b/Engine/Plugins/Experimental/CelestialVault/Source/CelestialVault/Public/ViewportCameraProvider.h
new file mode 100644
--- /dev/null
+++ b/Engine/Plugins/Experimental/CelestialVault/Source/CelestialVault/Public/ViewportCameraProvider.h
@@ -0,0 +1,30 @@
+﻿// Copyright Epic Games, Inc. All Rights Reserved.
+
+#pragma once
+
+#include "CoreMinimal.h"
+
+/**
+ * 
+ */
+class CELESTIALVAULT_API IViewportCameraProvider
+{
+public:
+	virtual ~IViewportCameraProvider() = default;
+
+	virtual bool GetCamera(FVector& OutLocation, FRotator& OutRotation) const = 0;
+};
+
+//////////////////////////////////////////////////////////////////////////
+// Default Runtime Implementation
+
+class CELESTIALVAULT_API FRuntimeCameraProvider : public IViewportCameraProvider
+{
+public:
+	virtual bool GetCamera(FVector& OutLocation, FRotator& OutRotation) const override;
+};
+
+
+// Global accessor
+CELESTIALVAULT_API IViewportCameraProvider* GetCameraProvider();
+CELESTIALVAULT_API void SetCameraProvider(IViewportCameraProvider* InProvider);
diff --git a/Engine/Plugins/Experimental/CelestialVault/Source/CelestialVaultEditor/Private/CelestialVaultEditor.cpp b/Engine/Plugins/Experimental/CelestialVault/Source/CelestialVaultEditor/Private/CelestialVaultEditor.cpp
--- a/Engine/Plugins/Experimental/CelestialVault/Source/CelestialVaultEditor/Private/CelestialVaultEditor.cpp
+++ b/Engine/Plugins/Experimental/CelestialVault/Source/CelestialVaultEditor/Private/CelestialVaultEditor.cpp
@@ -4,16 +4,23 @@
 
 #include "DaylightSavingsCustomization.h"
 #include "DaylightSavings.h"
+#include "ViewportCameraProvider.h"
+#include "EditorViewportCameraProvider.h"
 #include "PropertyEditorModule.h"
 #include "Modules/ModuleManager.h"
 
 #define LOCTEXT_NAMESPACE "FCelestialVaultEditorModule"
 
+
+FEditorViewportCameraProvider FCelestialVaultEditorModule::EditorProvider;
+
 void FCelestialVaultEditorModule::StartupModule()
 {
 	FPropertyEditorModule& PropertyModule = FModuleManager::LoadModuleChecked<FPropertyEditorModule>("PropertyEditor");
 
 	PropertyModule.RegisterCustomPropertyTypeLayout( FDaylightSavingsRule::StaticStruct()->GetFName(), FOnGetPropertyTypeCustomizationInstance::CreateStatic(&FDaylightSavingsCustomization::MakeInstance));
+
+	SetCameraProvider(&EditorProvider);
 }
 
 void FCelestialVaultEditorModule::ShutdownModule()
@@ -24,6 +31,9 @@ void FCelestialVaultEditorModule::ShutdownModule()
 
 		PropertyModule.UnregisterCustomPropertyTypeLayout(FDaylightSavingsRule::StaticStruct()->GetFName());
 	}
+
+	SetCameraProvider(nullptr);
+	
 }
 
 #undef LOCTEXT_NAMESPACE
diff --git a/Engine/Plugins/Experimental/CelestialVault/Source/CelestialVaultEditor/Private/CelestialVaultEditorUtilities.cpp b/Engine/Plugins/Experimental/CelestialVault/Source/CelestialVaultEditor/Private/CelestialVaultEditorUtilities.cpp
--- a/Engine/Plugins/Experimental/CelestialVault/Source/CelestialVaultEditor/Private/CelestialVaultEditorUtilities.cpp
+++ b/Engine/Plugins/Experimental/CelestialVault/Source/CelestialVaultEditor/Private/CelestialVaultEditorUtilities.cpp
@@ -2,8 +2,14 @@
 
 
 #include "CelestialVaultEditorUtilities.h"
+#include "Engine/Texture2D.h"
+#include "TextureResource.h"
+#include "Engine/TextureDefines.h"
+
 #include <LevelEditorViewport.h>
 
+
+
 void UCelestialVaultEditorUtilities::GetViewportCursorInformation(bool& Focused, FVector2D& ScreenLocation, FVector& WorldLocation, FVector& WorldDirection)
 {
 	Focused = false;
@@ -23,3 +29,85 @@ void UCelestialVaultEditorUtilities::GetViewportCursorInformation(bool& Focused,
 		Focused = true;
 	}
 }
+
+bool UCelestialVaultEditorUtilities::ComputeTextureMeanLuminance(UTexture2D* Texture, float& OutMean)
+{
+	OutMean = -1;
+	
+    if (!Texture || !Texture->GetPlatformData() || Texture->GetPlatformData()->Mips.Num() == 0)
+        return false;
+
+    FTexture2DMipMap& Mip = Texture->GetPlatformData()->Mips[0];
+
+    const int32 Width = Mip.SizeX;
+    const int32 Height = Mip.SizeY;
+
+    if (Width == 0 || Height == 0)
+        return false;
+
+    const EPixelFormat Format = Texture->GetPixelFormat();
+    if (Format == PF_G8 || Format == PF_B8G8R8A8)
+    {
+    	FByteBulkData* RawData = &Mip.BulkData;
+    	const void* DataPtr = RawData->LockReadOnly();
+    	if (!DataPtr)
+    		return false;
+
+    	const int64 PixelCount = static_cast<int64>(Width) * static_cast<int64>(Height);
+    	double TotalLuminance = 0.0;
+
+    	const bool bSRGB = Texture->SRGB;
+    	if (Format == PF_G8)
+    	{
+    		const uint8* Pixels = static_cast<const uint8*>(DataPtr);
+
+    		for (int64 i = 0; i < PixelCount; ++i)
+    		{
+    			float value = Pixels[i];
+    			float linear;
+
+    			if (bSRGB)
+    			{
+    				// Convert sRGB -> linear
+    				linear = FLinearColor::FromSRGBColor(FColor(Pixels[i], Pixels[i], Pixels[i])).R;
+    			}
+    			else
+    			{
+    				constexpr float Inv255 = 1.0f / 255.0f;
+    				linear = value * Inv255;
+    			}
+
+    			TotalLuminance += linear;
+    		}
+    	}
+    	else if (Format == PF_B8G8R8A8)
+    	{
+    		const FColor* Pixels = static_cast<const FColor*>(DataPtr);
+
+    		for (int64 i = 0; i < PixelCount; ++i)
+    		{
+    			FLinearColor linear;
+
+    			if (bSRGB)
+    			{
+    				// Converts from sRGB to Linear
+    				linear = FLinearColor(Pixels[i]);
+    			}
+    			else
+    			{
+    				linear = Pixels[i].ReinterpretAsLinear();
+    			}
+
+    			// Rec.709 luminance
+    			double Lum = 0.2126 * linear.R + 0.7152 * linear.G + 0.0722 * linear.B;
+    			TotalLuminance += Lum;
+    		}
+    	}
+
+    	RawData->Unlock();
+    	OutMean = static_cast<float>(TotalLuminance / static_cast<double>(PixelCount));
+    	return true;
+    }
+	
+    return false;
+}
diff --git a/Engine/Plugins/Experimental/CelestialVault/Source/CelestialVaultEditor/Private/EditorViewportCameraProvider.cpp b/Engine/Plugins/Experimental/CelestialVault/Source/CelestialVaultEditor/Private/EditorViewportCameraProvider.cpp
new file mode 100644
--- /dev/null
+++ b/Engine/Plugins/Experimental/CelestialVault/Source/CelestialVaultEditor/Private/EditorViewportCameraProvider.cpp
@@ -0,0 +1,26 @@
+﻿// Copyright Epic Games, Inc. All Rights Reserved.
+
+
+#include "EditorViewportCameraProvider.h"
+
+#include "Editor.h"
+#include "Subsystems/UnrealEditorSubsystem.h"
+
+bool FEditorViewportCameraProvider::GetCamera(FVector& OutLocation, FRotator& OutRotation) const
+{
+	if (!GEditor) return false;
+
+	// If we're playing (PIE), fallback to runtime behavior
+	if (GEditor->GetPIEWorldContext())
+	{
+		static FRuntimeCameraProvider DefaultProvider;
+		
+		return DefaultProvider.GetCamera(OutLocation, OutRotation);
+	}
+
+	if (UUnrealEditorSubsystem* EditorSubsystem = GEditor->GetEditorSubsystem<UUnrealEditorSubsystem>())
+	{
+		return EditorSubsystem->GetLevelViewportCameraInfo(OutLocation, OutRotation);
+	}
+	return false;
+}
\ No newline at end of file
diff --git a/Engine/Plugins/Experimental/CelestialVault/Source/CelestialVaultEditor/Public/CelestialVaultEditor.h b/Engine/Plugins/Experimental/CelestialVault/Source/CelestialVaultEditor/Public/CelestialVaultEditor.h
--- a/Engine/Plugins/Experimental/CelestialVault/Source/CelestialVaultEditor/Public/CelestialVaultEditor.h
+++ b/Engine/Plugins/Experimental/CelestialVault/Source/CelestialVaultEditor/Public/CelestialVaultEditor.h
@@ -5,9 +5,14 @@
 #include "CoreMinimal.h"
 #include "Modules/ModuleManager.h"
 
+class FEditorViewportCameraProvider;
+
 class FCelestialVaultEditorModule : public IModuleInterface
 {
 public:
     virtual void StartupModule() override;
     virtual void ShutdownModule() override;
+
+private:
+    static FEditorViewportCameraProvider EditorProvider;
 };
diff --git a/Engine/Plugins/Experimental/CelestialVault/Source/CelestialVaultEditor/Public/CelestialVaultEditorUtilities.h b/Engine/Plugins/Experimental/CelestialVault/Source/CelestialVaultEditor/Public/CelestialVaultEditorUtilities.h
--- a/Engine/Plugins/Experimental/CelestialVault/Source/CelestialVaultEditor/Public/CelestialVaultEditorUtilities.h
+++ b/Engine/Plugins/Experimental/CelestialVault/Source/CelestialVaultEditor/Public/CelestialVaultEditorUtilities.h
@@ -3,6 +3,7 @@
 #pragma once
 
 #include "CoreMinimal.h"
+#include "Engine/Texture2D.h"
 #include "Kismet/BlueprintFunctionLibrary.h"
 #include "CelestialVaultEditorUtilities.generated.h"
 
@@ -24,4 +25,7 @@ class CELESTIALVAULTEDITOR_API UCelestialVaultEditorUtilities : public UBlueprin
      */
 	UFUNCTION(BlueprintCallable, Category = "Celestial Vault|Editor Utilities")
 	static void GetViewportCursorInformation(bool& Focused, FVector2D& ScreenLocation, FVector& WorldLocation, FVector& WorldDirection);
+
+	UFUNCTION(BlueprintCallable, Category = "Celestial Vault|Editor Utilities")
+	static bool ComputeTextureMeanLuminance(UTexture2D* Texture, float& OutMean);
 };
diff --git a/Engine/Plugins/Experimental/CelestialVault/Source/CelestialVaultEditor/Public/EditorViewportCameraProvider.h b/Engine/Plugins/Experimental/CelestialVault/Source/CelestialVaultEditor/Public/EditorViewportCameraProvider.h
new file mode 100644
--- /dev/null
+++ b/Engine/Plugins/Experimental/CelestialVault/Source/CelestialVaultEditor/Public/EditorViewportCameraProvider.h
@@ -0,0 +1,11 @@
+﻿// Copyright Epic Games, Inc. All Rights Reserved.
+
+#pragma once
+
+#include "ViewportCameraProvider.h"
+
+class FEditorViewportCameraProvider : public IViewportCameraProvider
+{
+public:
+	virtual bool GetCamera(FVector& OutLocation, FRotator& OutRotation) const override;
+};
