Minor visual tweaks

#virtualized
#Synced-CL 51124650

diff --git a/Engine/Plugins/Experimental/CelestialVault/Source/CelestialVault/Private/CelestialVaultDaySequenceActor.cpp b/Engine/Plugins/Experimental/CelestialVault/Source/CelestialVault/Private/CelestialVaultDaySequenceActor.cpp
--- a/Engine/Plugins/Experimental/CelestialVault/Source/CelestialVault/Private/CelestialVaultDaySequenceActor.cpp
+++ b/Engine/Plugins/Experimental/CelestialVault/Source/CelestialVault/Private/CelestialVaultDaySequenceActor.cpp
@@ -68,7 +68,7 @@ ACelestialVaultDaySequenceActor::ACelestialVaultDaySequenceActor(const FObjectIn
 	{
 		GlobalPostProcessVolume->SetupAttachment(RootComponent);
 		GlobalPostProcessVolume->Settings.bOverride_AutoExposureMinBrightness = true;
-		GlobalPostProcessVolume->Settings.AutoExposureMinBrightness = 0.25f;
+		GlobalPostProcessVolume->Settings.AutoExposureMinBrightness = -2.0f;
 		GlobalPostProcessVolume->Settings.bOverride_LocalExposureHighlightContrastCurve = true;
 		GlobalPostProcessVolume->Settings.LocalExposureHighlightContrastCurve = HighlightContrastCurve.Object;
 
@@ -204,6 +204,7 @@ ACelestialVaultDaySequenceActor::ACelestialVaultDaySequenceActor(const FObjectIn
 		StarsComponent->SetAffectDynamicIndirectLighting(false);
 		StarsComponent->SetCanEverAffectNavigation(false);
 		StarsComponent->SetMaterial(0, StarsDefaultMaterial.Object);
+		StarsComponent->SetTranslucentSortPriority(-1);
 		StarsComponent->bUseAttachParentBound = true;	
 	}
 	
